Dev to Main from cl #2110143

diff --git a/Engine/Source/Developer/FunctionalTesting/Private/FunctionalTest.cpp b/Engine/Source/Developer/FunctionalTesting/Private/FunctionalTest.cpp
--- a/Engine/Source/Developer/FunctionalTesting/Private/FunctionalTest.cpp
+++ b/Engine/Source/Developer/FunctionalTesting/Private/FunctionalTest.cpp
@@ -250,7 +250,7 @@ void AFunctionalTest::GoToObservationPoint()
 		APlayerController* PC = World->GetFirstPlayerController();
 		if (PC && PC->GetPawn())
 		{
-			PC->GetPawn()->TeleportTo(ObservationPoint->GetActorLocation(), ObservationPoint->GetActorRotation());
+			PC->GetPawn()->TeleportTo(ObservationPoint->GetActorLocation(), ObservationPoint->GetActorRotation(), /*bIsATest=*/false, /*bNoCheck=*/true);
 			PC->SetControlRotation(ObservationPoint->GetActorRotation());
 		}
 	}
diff --git a/Engine/Source/Developer/LogVisualizer/Private/LogVisualizer.cpp b/Engine/Source/Developer/LogVisualizer/Private/LogVisualizer.cpp
--- a/Engine/Source/Developer/LogVisualizer/Private/LogVisualizer.cpp
+++ b/Engine/Source/Developer/LogVisualizer/Private/LogVisualizer.cpp
@@ -62,6 +62,16 @@ void FLogVisualizer::CloseUI(UWorld* InWorld)
 	}
 }
 
+bool FLogVisualizer::IsOpenUI(UWorld* InWorld)
+{
+	if (LogWindow.IsValid() && World.IsValid() && World == InWorld)
+	{
+		return true;
+	}
+
+	return false;
+}
+
 void FLogVisualizer::CleanUp()
 {
 	FVisualLog::Get()->ClearNewLogsObserver();
diff --git a/Engine/Source/Developer/LogVisualizer/Private/LogVisualizer.h b/Engine/Source/Developer/LogVisualizer/Private/LogVisualizer.h
--- a/Engine/Source/Developer/LogVisualizer/Private/LogVisualizer.h
+++ b/Engine/Source/Developer/LogVisualizer/Private/LogVisualizer.h
@@ -12,6 +12,7 @@ class FLogVisualizer : public ILogVisualizer
 public:
 	virtual void SummonUI(UWorld* InWorld) override;
 	virtual void CloseUI(UWorld* InWorld) override;
+	virtual bool IsOpenUI(UWorld* InWorld) override;
 
 	virtual bool IsRecording();
 	// End ILogVisualizer interface
diff --git a/Engine/Source/Developer/LogVisualizer/Private/SLogBar.cpp b/Engine/Source/Developer/LogVisualizer/Private/SLogBar.cpp
--- a/Engine/Source/Developer/LogVisualizer/Private/SLogBar.cpp
+++ b/Engine/Source/Developer/LogVisualizer/Private/SLogBar.cpp
@@ -3,6 +3,8 @@
 #include "LogVisualizerPCH.h"
 #include "SLogBar.h"
 
+#define LOCTEXT_NAMESPACE "SLogBar"
+
 const float SLogBar::SubPixelMinSize = 3.0f;
 const float SLogBar::TimeUnit = 1.f/60.f;
 const float SLogBar::MaxUnitSizePx = 16.f;
@@ -27,6 +29,7 @@ void SLogBar::Construct( const FArguments& InArgs )
 	TotalTime = 1.0;
 	RowIndex = INDEX_NONE;
 	bShouldDrawSelection = false;
+	HistogramPreviewWindow = 50;
 }
 
 int32 SLogBar::OnPaint( const FGeometry& AllottedGeometry, const FSlateRect& MyClippingRect, FSlateWindowElementList& OutDrawElements, int32 LayerId, const FWidgetStyle& InWidgetStyle, bool bParentEnabled ) const
@@ -59,6 +62,7 @@ int32 SLogBar::OnPaint( const FGeometry& AllottedGeometry, const FSlateRect& MyC
 	);	
 
 	// Draw all bars
+	bool bDrawHistogramCursor = false;
 	int32 EntryIndex = 0;
 	while (EntryIndex < Entries.Num())
 	{
@@ -70,6 +74,7 @@ int32 SLogBar::OnPaint( const FGeometry& AllottedGeometry, const FSlateRect& MyC
 				EntryIndex++;
 				continue;
 			}
+			bDrawHistogramCursor |= Entry->HistogramSamples.Num() > 0;
 		}
 
 		// find bar width, connect all contiguous bars to draw them as one geometry (rendering optimization)
@@ -110,29 +115,75 @@ int32 SLogBar::OnPaint( const FGeometry& AllottedGeometry, const FSlateRect& MyC
 
 		EntryIndex = StartIndex;
 	}
-	
-	// draw "current time"
+
+	if (Entries.Num() > 0)
 	{
-		const float RelativePos = Offset + (DisplayedTime.Execute() - StartTime) / TotalTime * Zoom;
-		const float PosX = (float)(LogBarWidth * RelativePos);
-		const float ClampedSize = FMath::Clamp(TimeUnit / TotalTime * Zoom, 0.0f, 1.0f );
-		float CurrentTimeMarkWidth = (float)FMath::Max(LogBarWidth * ClampedSize
-			, ClampedSize > 0.0f ? SubPixelMinSize : 0.0f);
-		
-		// Draw Event bar
-		FSlateDrawElement::MakeBox(
-			OutDrawElements,
-			RetLayerId++,
-			AllottedGeometry.ToPaintGeometry(
-			FVector2D( PosX - 2.f, 0.0f ),
-			FVector2D( CurrentTimeMarkWidth + 4.f, AllottedGeometry.Size.Y )),
-			TimeMarkImage,
-			ForegroundClippingRect,
-			DrawEffects,
-			CurrentTimeColor
-			);
+		// draw "current time"
+		{
+			const float RelativePos = Offset + (DisplayedTime.Execute() - StartTime) / TotalTime * Zoom;
+			const float PosX = (float)(LogBarWidth * RelativePos);
+			const float ClampedSize = FMath::Clamp(TimeUnit / TotalTime * Zoom, 0.0f, 1.0f);
+			float CurrentTimeMarkWidth = (float)FMath::Max(LogBarWidth * ClampedSize
+				, ClampedSize > 0.0f ? SubPixelMinSize : 0.0f);
+
+			// Draw Event bar
+			FSlateDrawElement::MakeBox(
+				OutDrawElements,
+				RetLayerId++,
+				AllottedGeometry.ToPaintGeometry(
+				FVector2D(PosX - 2.f, 0.0f),
+				FVector2D(CurrentTimeMarkWidth + 4.f, AllottedGeometry.Size.Y)),
+				FillImage,
+				ForegroundClippingRect,
+				DrawEffects,
+				CurrentTimeColor
+				);
+		}
+
+		// draw PreviewWindow for 2d graphs
+		int32 EntryIndex = CurrentEntryIndex.IsBound() ? CurrentEntryIndex.Execute() : INDEX_NONE;
+		if (bDrawHistogramCursor && Entries.IsValidIndex(EntryIndex) == true)
+		{
+			static FColor DrawColor = FColor(CurrentTimeColor.R, CurrentTimeColor.G, CurrentTimeColor.B, 128);
+			const float RelativePos = Offset + (DisplayedTime.Execute() - StartTime) / TotalTime * Zoom;
+			const float PosX = (float)(LogBarWidth * RelativePos);
+			const float ClampedSize = FMath::Clamp(TimeUnit / TotalTime * Zoom, 0.0f, 1.0f);
+
+			TSharedPtr<FVisLogEntry> Entry = Entries[EntryIndex];
+			float StartX, EndX;
+			CalculateEntryGeometry(Entry.Get(), AllottedGeometry, StartX, EndX);
+			const float WindowHalfWidth = LogBarWidth * HistogramPreviewWindow * 0.01;
+
+			const float RelativeStartPos = PosX - WindowHalfWidth * 0.5f;
+
+			// Draw Event bar
+			FSlateDrawElement::MakeBox(
+				OutDrawElements,
+				RetLayerId++,
+				AllottedGeometry.ToPaintGeometry(
+				FVector2D(RelativeStartPos, 0.0f),
+				FVector2D(WindowHalfWidth, AllottedGeometry.Size.Y)),
+				FillImage,
+				ForegroundClippingRect,
+				DrawEffects,
+				DrawColor
+				);
+
+			// Draw Event bar
+			FSlateDrawElement::MakeBox(
+				OutDrawElements,
+				RetLayerId++,
+				AllottedGeometry.ToPaintGeometry(
+				FVector2D(RelativeStartPos, 0.0f),
+				FVector2D(WindowHalfWidth, AllottedGeometry.Size.Y)),
+				SelectedImage,
+				ForegroundClippingRect,
+				DrawEffects,
+				DrawColor
+				);
+		}
 	}
-	
+
 	if (bShouldDrawSelection)
 	{
 		EntryIndex = CurrentEntryIndex.IsBound() ? CurrentEntryIndex.Execute() : INDEX_NONE;
@@ -335,6 +386,31 @@ void SLogBar::SetEntries(const TArray<TSharedPtr<FVisLogEntry> >& InEntries, flo
 	Entries = InEntries;
 	StartTime = InStartTime;
 	TotalTime = InTotalTime;
+
+	bool bDrawHistogramCursor = false;
+	int32 EntryIndex = 0;
+	while (EntryIndex < Entries.Num())
+	{
+		TSharedPtr<FVisLogEntry> Entry = Entries[EntryIndex];
+		if (!Entry.IsValid())
+		{
+			EntryIndex++;
+			continue;
+		}
+
+		bDrawHistogramCursor |= Entry->HistogramSamples.Num() > 0;
+		EntryIndex++;
+	}
+
+	if (bDrawHistogramCursor)
+	{
+		//const FText ToolTip(TEXT("Use mouse wheel and left shift to modify graph window"));
+		SetToolTipText(TAttribute<FString>(TEXT("Use mouse wheel and left shift to modify 2d graph cursor size")));
+	}
+	else
+	{
+		SetToolTipText(FText());
+	}
 }
 
 void SLogBar::OnCurrentTimeChanged(float NewTime)
@@ -360,6 +436,11 @@ void SLogBar::SetZoomAndOffset(float InZoom, float InOffset)
 	SetOffset(InOffset);
 } 
 
+void SLogBar::SetHistogramWindow(float InWindowSize)
+{
+	HistogramPreviewWindow = InWindowSize;
+}
+
 void SLogBar::UpdateShouldDrawSelection()
 {
 	/*const bool bNewShouldDraw = ShouldDrawSelection.Execute();
@@ -374,3 +455,5 @@ void SLogBar::UpdateShouldDrawSelection()
 	}*/
 }
 
+#undef LOCTEXT_NAMESPACE 
+
diff --git a/Engine/Source/Developer/LogVisualizer/Private/SLogBar.h b/Engine/Source/Developer/LogVisualizer/Private/SLogBar.h
--- a/Engine/Source/Developer/LogVisualizer/Private/SLogBar.h
+++ b/Engine/Source/Developer/LogVisualizer/Private/SLogBar.h
@@ -134,6 +134,10 @@ public:
 
 	void SetZoomAndOffset(float InZoom, float InOffset);
 
+	void SetHistogramWindow(float InWindowSize);
+
+	float GetHistogramPreviewWindow() { return HistogramPreviewWindow; }
+
 	/**
 	 * Gets the graph's offset value
 	 *
@@ -201,6 +205,8 @@ private:
 	/** Current offset of the graph */
 	float Offset;
 
+	float HistogramPreviewWindow;
+
 	/** Last hovered event index */
 	int32 LastHoveredEvent;
 	
diff --git a/Engine/Source/Developer/LogVisualizer/Private/SLogVisualizer.cpp b/Engine/Source/Developer/LogVisualizer/Private/SLogVisualizer.cpp
--- a/Engine/Source/Developer/LogVisualizer/Private/SLogVisualizer.cpp
+++ b/Engine/Source/Developer/LogVisualizer/Private/SLogVisualizer.cpp
@@ -40,6 +40,21 @@ namespace LogVisualizer
 	static const FString FileTypes = FString::Printf( TEXT("%s (%s)|%s"), *LogFileDescription, *LogFileExtension, *LogFileExtension );
 }
 
+template <typename ItemType>
+class SLogListView : public SListView<ItemType>
+{
+public:
+	virtual FReply OnMouseWheel(const FGeometry& MyGeometry, const FPointerEvent& MouseEvent) OVERRIDE
+	{
+		if (!MouseEvent.IsLeftShiftDown())
+		{
+			return SListView<ItemType>::OnMouseWheel(MyGeometry, MouseEvent);
+		};
+
+		return FReply::Unhandled();
+	}
+};
+
 FColor SLogVisualizer::ColorPalette[] = {
 	FColor(0xff00A480),
 	FColorList::Aquamarine,
@@ -52,6 +67,12 @@ FColor SLogVisualizer::ColorPalette[] = {
 	FColorList::NeonBlue,
 	FColorList::MediumSlateBlue,
 	FColorList::SpicyPink,
+	FColorList::SpringGreen,
+	FColorList::SteelBlue,
+	FColorList::SummerSky,
+	FColorList::Violet,
+	FColorList::VioletRed,
+	FColorList::YellowGreen,
 	FColor(0xff62E200),
 	FColor(0xff1F7B67),
 	FColor(0xff62AA2A),
@@ -99,6 +120,9 @@ void SLogVisualizer::Construct(const FArguments& InArgs, FLogVisualizer* InLogVi
 	CurrentViewedTime = 0.f;
 	bDrawLogEntriesPath = true;
 	bIgnoreTrivialLogs = true;
+	HistogramPreviewWindow = 50;
+	bShowHistogramLabelsOutside = false;
+	bStickToLastData = false;
 
 	UsedCategories.Empty();
 
@@ -253,6 +277,32 @@ void SLogVisualizer::Construct(const FArguments& InArgs, FLogVisualizer* InLogVi
 							.ToolTipText(LOCTEXT("VisLogIgnoreTrivialLogsTooltip", "Whether to show trivial logs, i.e. the ones with only one entry."))
 						]
 					]
+					+ SVerticalBox::Slot()
+						.AutoHeight()
+						[
+							SNew(SCheckBox)
+							.OnCheckStateChanged(this, &SLogVisualizer::OnChangeHistogramLabelLocation)
+							.IsChecked(this, &SLogVisualizer::GetHistogramLabelLocation)
+							.Content()
+							[
+								SNew(STextBlock)
+								.Text(LOCTEXT("VisLogHistogramLabelLocation", "Set histogram labels outside graph"))
+								.ToolTipText(LOCTEXT("VisLogHistogramLabelLocationTooltip", "Whether to show histogram labels inside graph or outside."))
+							]
+						]
+					+ SVerticalBox::Slot()
+						.AutoHeight()
+						[
+							SNew(SCheckBox)
+							.OnCheckStateChanged(this, &SLogVisualizer::OnStickToLastData)
+							.IsChecked(this, &SLogVisualizer::GetStickToLastData)
+							.Content()
+							[
+								SNew(STextBlock)
+								.Text(LOCTEXT("VisLogStickToLastData", "Stick to recent data"))
+								.ToolTipText(LOCTEXT("VisLogStickToLastDataTooltip", "Whether to show the recent data or not."))
+							]
+						]
 				]
 			]
 
@@ -281,7 +331,7 @@ void SLogVisualizer::Construct(const FArguments& InArgs, FLogVisualizer* InLogVi
 					.BorderImage(FEditorStyle::GetBrush("Menu.Background"))
 					.Padding(1.0)
 					[
-						SAssignNew(LogsListWidget, SListView< TSharedPtr<FLogsListItem> >)
+						SAssignNew(LogsListWidget, SLogListView< TSharedPtr<FLogsListItem> >)
 						.ItemHeight(20)
 						// Called when the user double-clicks with LMB on an item in the list
 						.OnMouseButtonDoubleClick(this, &SLogVisualizer::OnListDoubleClick)
@@ -638,6 +688,16 @@ void SLogVisualizer::Tick( const FGeometry& AllottedGeometry, const double InCur
 	{
 		DoFullUpdate();
 	}
+
+	if (bStickToLastData && World && !World->bPlayersOnly  && LogVisualizer->IsRecording()  && World->IsPlayInEditor() && LogVisualizer && LogVisualizer->Logs.IsValidIndex(SelectedLogIndex))
+	{
+		TSharedPtr<FActorsVisLog> Log = LogVisualizer->Logs[SelectedLogIndex];
+		if (Log->Entries.Num() > 0 && LogEntryIndex != Log->Entries.Num() - 1)
+		{
+			LogEntryIndex = Log->Entries.Num() - 1;
+			ShowEntry(Log->Entries[LogEntryIndex].Get());
+		}
+	}
 }
 
 FReply SLogVisualizer::OnMouseWheel( const FGeometry& MyGeometry, const FPointerEvent& MouseEvent )
@@ -647,6 +707,13 @@ FReply SLogVisualizer::OnMouseWheel( const FGeometry& MyGeometry, const FPointer
 		OnSetZoomValue(FMath::Clamp(ZoomSliderValue + MouseEvent.GetWheelDelta() * 0.05f, 0.f, 1.f));
 		return FReply::Handled();
 	}
+
+	if (MouseEvent.IsLeftShiftDown())
+	{
+		OnSetHistogramWindowValue(MouseEvent.GetWheelDelta());
+		return FReply::Handled();
+	}
+
 	return SCompoundWidget::OnMouseWheel(MyGeometry, MouseEvent);
 }
 
@@ -852,8 +919,11 @@ void SLogVisualizer::DoFullUpdate()
 		if (LogListItem->IsValid() && LogVisualizer->Logs.IsValidIndex((*LogListItem)->LogIndex))
 		{
 			TSharedPtr<FActorsVisLog>& Log = LogVisualizer->Logs[(*LogListItem)->LogIndex];
-			LogsStartTime = FMath::Min(Log->Entries[0]->TimeStamp, LogsStartTime);
-			LogsEndTime = FMath::Max(Log->Entries[Log->Entries.Num()-1]->TimeStamp, LogsEndTime);
+			if (Log->Entries.Num() > 0)
+			{
+				LogsStartTime = FMath::Min(Log->Entries[0]->TimeStamp, LogsStartTime);
+				LogsEndTime = FMath::Max(Log->Entries[Log->Entries.Num() - 1]->TimeStamp, LogsEndTime);
+			}
 		}
 	}
 
@@ -1019,6 +1089,21 @@ void SLogVisualizer::SetCurrentViewedTime(float NewTime, const bool bForce)
 	}
 
 	CurrentViewedTime = NewTime;
+
+#if WITH_EDITOR
+	UEditorEngine *EEngine = Cast<UEditorEngine>(GEngine);
+	if (GIsEditor && EEngine != NULL)
+	{
+		for (int32 i = 0; i < EEngine->AllViewportClients.Num(); i++)
+		{
+			FEditorViewportClient* ViewportClient = EEngine->AllViewportClients[i];
+			if (ViewportClient)
+			{
+				ViewportClient->Invalidate();
+			}
+		}
+	}
+#endif
 }
 
 void SLogVisualizer::RequestShowLogEntry(TSharedPtr<FLogsListItem> Item, TSharedPtr<FVisLogEntry> LogEntry)
@@ -1054,7 +1139,14 @@ FLinearColor SLogVisualizer::GetColorForUsedCategory(int32 Index)
 		return ColorPalette[Index];
 	}
 
-	return FLinearColor::White;
+	static bool bReateColorList = false;
+	static FColorList StaticColor;
+	if (!bReateColorList)
+	{
+		bReateColorList = true;
+		StaticColor.CreateColorMap();
+	}
+	return StaticColor.GetFColorByIndex(Index);
 }
 
 TSharedRef<ITableRow> SLogVisualizer::HandleGenerateLogStatus(TSharedPtr<FLogStatusItem> InItem, const TSharedRef<STableViewBase>& OwnerTable)
@@ -1290,6 +1382,12 @@ void SLogVisualizer::OnZoomScrolled(float InScrollOffsetFraction)
 	}
 }
 
+void SLogVisualizer::OnSetHistogramWindowValue(float NewValue)
+{
+	HistogramPreviewWindow = FMath::Clamp(HistogramPreviewWindow + NewValue * 1.0f, 0.0f, 100.0f);
+	HistogramWindowChangedNotify.Broadcast(HistogramPreviewWindow);
+}
+
 void SLogVisualizer::OnDrawLogEntriesPathChanged(ESlateCheckBoxState::Type NewState)
 {
 	bDrawLogEntriesPath = (NewState == ESlateCheckBoxState::Checked);
@@ -1311,6 +1409,26 @@ ESlateCheckBoxState::Type SLogVisualizer::GetIgnoreTrivialLogs() const
 	return bIgnoreTrivialLogs ? ESlateCheckBoxState::Checked : ESlateCheckBoxState::Unchecked;
 }
 
+void SLogVisualizer::OnChangeHistogramLabelLocation(ESlateCheckBoxState::Type NewState)
+{
+	bShowHistogramLabelsOutside = (NewState == ESlateCheckBoxState::Checked);
+}
+
+ESlateCheckBoxState::Type SLogVisualizer::GetHistogramLabelLocation() const
+{
+	return bShowHistogramLabelsOutside ? ESlateCheckBoxState::Checked : ESlateCheckBoxState::Unchecked;
+}
+
+void SLogVisualizer::OnStickToLastData(ESlateCheckBoxState::Type NewState)
+{
+	bStickToLastData = (NewState == ESlateCheckBoxState::Checked);
+}
+
+ESlateCheckBoxState::Type SLogVisualizer::GetStickToLastData() const
+{
+	return bStickToLastData ? ESlateCheckBoxState::Checked : ESlateCheckBoxState::Unchecked;
+}
+
 void SLogVisualizer::OnToggleCamera(ESlateCheckBoxState::Type NewState)
 {
 	UWorld* World = GetWorld();
@@ -1374,6 +1492,32 @@ void SLogVisualizer::DrawOnCanvas(UCanvas* Canvas, APlayerController*)
 			Canvas->DrawText(Font, TimeStampString, EntryScreenLoc.X, EntryScreenLoc.Y);
 
 			//let's draw histogram data
+
+			static const float StartGoldenRatio[] = {
+				FMath::SRand(),
+				FMath::SRand(),
+				FMath::SRand(),
+				FMath::SRand(),
+				FMath::SRand(),
+				FMath::SRand(),
+				FMath::SRand(),
+				FMath::SRand(),
+				FMath::SRand(),
+				FMath::SRand(),
+				FMath::SRand(),
+				FMath::SRand(),
+				FMath::SRand(),
+				FMath::SRand(),
+				FMath::SRand(),
+				FMath::SRand(),
+				FMath::SRand(),
+				FMath::SRand(),
+				FMath::SRand(),
+				FMath::SRand(),
+			};
+
+			const float GoldenRatioConjugate = 0.618033988749895f;
+
 			struct FGraphLineData
 			{
 				FName DataName;
@@ -1391,14 +1535,37 @@ void SLogVisualizer::DrawOnCanvas(UCanvas* Canvas, APlayerController*)
 
 			TMap<FName, FGraphData>	CollectedGraphs;
 
+			float MinTime, MaxTime;
+			Timeline->GetMinMaxValues(MinTime, MaxTime);
+
+			const float StartTime = Timeline->GetOffset() + MinTime;
+			const float EndTime = StartTime + (MaxTime - MinTime) / this->GetZoom();
+			const float WindowHalfWidth = (EndTime - StartTime) * HistogramPreviewWindow * 0.01 * 0.5;
+			const FVector2D TimeStampWindow(Entry->TimeStamp - WindowHalfWidth, Entry->TimeStamp + WindowHalfWidth);
+
+			int32 ColorIndex = 0;
 			for (int32 EntryIndex = 0; EntryIndex < Entries.Num(); ++EntryIndex)
 			{
 				const TSharedPtr<FVisLogEntry>& CurrentEntry = Entries[EntryIndex];
-				if (CurrentEntry->TimeStamp > Entry->TimeStamp)
+				if (HistogramPreviewWindow <= 0)
 				{
-					break;
+					if (CurrentEntry->TimeStamp > Entry->TimeStamp)
+					{
+						break;
+					}
 				}
+				else
+				{
+					if (CurrentEntry->TimeStamp < TimeStampWindow.X)
+					{
+						continue;
+					}
 
+					if (CurrentEntry->TimeStamp > TimeStampWindow.Y)
+					{
+						break;
+					}
+				}
 				const int32 SamplesNum = CurrentEntry->HistogramSamples.Num();
 				for (int32 SampleIndex = 0; SampleIndex < SamplesNum; ++SampleIndex)
 				{
@@ -1423,6 +1590,12 @@ void SLogVisualizer::DrawOnCanvas(UCanvas* Canvas, APlayerController*)
 			int32 GraphIndex = 0;
 			if (CollectedGraphs.Num() > 0)
 			{
+				const float GraphWidth = 0.8f / CollectedGraphs.Num();
+				const float GraphSpacing = 0.2f / (CollectedGraphs.Num()+1);
+				float CurrentX = GraphSpacing;
+				float CurrentY = (1.0f - GraphWidth) * 0.5f;
+				int32 GraphIndex = 0;
+
 				for (auto It(CollectedGraphs.CreateConstIterator()); It; ++It)
 				{
 					TWeakObjectPtr<UReporterGraph> HistogramGraph = Canvas->GetReporterGraph();
@@ -1432,32 +1605,43 @@ void SLogVisualizer::DrawOnCanvas(UCanvas* Canvas, APlayerController*)
 					}
 					HistogramGraph->SetNumGraphLines(It->Value.GraphLines.Num());
 					int32 LineIndex = 0;
+					UFont* Font = GEngine->GetSmallFont();
+					int32 MaxStringSize = 0;
+					float Hue = 0;// StartGoldenRatio[GraphIndex++];
+
 					for (auto LinesIt(It->Value.GraphLines.CreateConstIterator()); LinesIt; ++LinesIt)
 					{
-						int32 ColorIndex = UsedCategories.Find(LinesIt->Value.DataName.ToString());
-						if (ColorIndex == INDEX_NONE)
+						Hue += GoldenRatioConjugate;
+						if (Hue > 1)
 						{
-							ColorIndex = UsedCategories.AddUnique(LinesIt->Value.DataName.ToString());
+							Hue -= FMath::FloorToFloat(Hue);
 						}
-
-						HistogramGraph->GetGraphLine(LineIndex)->Color = GetColorForUsedCategory(ColorIndex);
+						HistogramGraph->GetGraphLine(LineIndex)->Color = FLinearColor::FGetHSV(Hue * 255, 0, 244);
 						HistogramGraph->GetGraphLine(LineIndex)->LineName = LinesIt->Value.DataName.ToString();
 						HistogramGraph->GetGraphLine(LineIndex)->Data.Append(LinesIt->Value.Samples);
+
+						int32 DummyY, CurrentX;
+						StringSize(Font, CurrentX, DummyY, *LinesIt->Value.DataName.ToString());
+						MaxStringSize = CurrentX > MaxStringSize ? CurrentX : MaxStringSize;
+
 						++LineIndex;
 					}
 
-					float dx = 0.8f / CollectedGraphs.Num();
-					HistogramGraph->SetGraphScreenSize( 0.1f + GraphIndex * (dx+0.02), 0.1f + GraphIndex * dx + dx, (1.0f - dx) * 0.5f,   1.0f - (1 - dx) * 0.5f);
-					HistogramGraph->SetAxesMinMax(It->Value.Min, It->Value.Max);
+					FVector2D GraphSpaceSize;
+					GraphSpaceSize.Y = GraphSpaceSize.X = 0.8f / CollectedGraphs.Num();
 
+					HistogramGraph->SetGraphScreenSize(CurrentX, CurrentX + GraphWidth, CurrentY, CurrentY + GraphWidth);
+					CurrentX += GraphWidth + GraphSpacing;
+					HistogramGraph->SetAxesMinMax(FVector2D(TimeStampWindow.X, It->Value.Min.Y), FVector2D(TimeStampWindow.Y, It->Value.Max.Y));
+					
 					HistogramGraph->SetNumThresholds(1);
 					FGraphThreshold* GraphThreshold = HistogramGraph->GetThreshold(0);
 					GraphThreshold->Threshold = 0.f;
 					GraphThreshold->Color = FLinearColor::White;
 					GraphThreshold->ThresholdName = TEXT("     0");
 					HistogramGraph->SetStyles(EGraphAxisStyle::Grid, EGraphDataStyle::Lines);
-					HistogramGraph->SetBackgroundColor( FColor(0,0,0, 90) );
-					HistogramGraph->SetLegendPosition(ELegendPosition::Inside);
+					HistogramGraph->SetBackgroundColor( FColor(0,0,0, 200) );
+					HistogramGraph->SetLegendPosition(bShowHistogramLabelsOutside ? ELegendPosition::Outside : ELegendPosition::Inside);
 
 					HistogramGraph->bVisible = true;
 					HistogramGraph->Draw(Canvas);
diff --git a/Engine/Source/Developer/LogVisualizer/Private/SLogVisualizer.h b/Engine/Source/Developer/LogVisualizer/Private/SLogVisualizer.h
--- a/Engine/Source/Developer/LogVisualizer/Private/SLogVisualizer.h
+++ b/Engine/Source/Developer/LogVisualizer/Private/SLogVisualizer.h
@@ -65,15 +65,16 @@ public:
 	static const FName NAME_LogTimeSpan;
 
 	/** Expressed in Hz */
-	static const int32 FullUpdateFrequency = 2; 
+	static const int32 FullUpdateFrequency = 2;
 
 	DECLARE_MULTICAST_DELEGATE_TwoParams(FOnZoomChanged, float, float);
+	DECLARE_MULTICAST_DELEGATE_OneParam(FOnHistogramWindowChanged, float);
 
 	SLATE_BEGIN_ARGS(SLogVisualizer) {}
 
 	SLATE_END_ARGS()
 
-	void Construct(const FArguments& InArgs, FLogVisualizer* InAnalyzer);
+		void Construct(const FArguments& InArgs, FLogVisualizer* InAnalyzer);
 
 	virtual ~SLogVisualizer();
 
@@ -100,9 +101,11 @@ public:
 	TSharedRef<ITableRow> LogsListGenerateRow(TSharedPtr<FLogsListItem> InItem, const TSharedRef<STableViewBase>& OwnerTable);
 	void LogsListSelectionChanged(TSharedPtr<FLogsListItem> SelectedItem, ESelectInfo::Type SelectInfo);
 	TSharedRef<ITableRow> LogEntryLinesGenerateRow(TSharedPtr<FLogEntryItem> Item, const TSharedRef<STableViewBase>& OwnerTable);
-	
+
 	FOnZoomChanged& OnZoomChanged() { return ZoomChangedNotify; }
 
+	FOnHistogramWindowChanged& OnHistogramWindowChanged() { return HistogramWindowChangedNotify; }
+
 	void OnListDoubleClick(TSharedPtr<FLogsListItem>);
 	void DoFullUpdate();
 	void RequestFullUpdate() { TimeTillNextUpdate = 0; }
@@ -140,6 +143,7 @@ public:
 	float GetCurrentViewedTime() const { return CurrentViewedTime; }
 
 	static FLinearColor GetColorForUsedCategory(int32 Index);
+	float GetHistogramPreviewWindow() { return HistogramPreviewWindow; }
 
 protected:
 	float GetMaxScrollOffsetFraction() const
@@ -152,6 +156,7 @@ protected:
 		return GetZoom() - 1.0f;
 	}
 
+
 	void DrawOnCanvas(UCanvas* Canvas, APlayerController*);
 
 	void LoadFiles(TArray<FString>& OpenFilenames);
@@ -189,6 +194,8 @@ private:
 	void OnSetZoomValue(float NewValue);
 	void OnZoomScrolled(float InScrollOffsetFraction);
 
+	void OnSetHistogramWindowValue(float NewValue);
+
 	/** handler for toggling DrawLogEntriesPath option change */
 	void OnDrawLogEntriesPathChanged(ESlateCheckBoxState::Type NewState);
 	/** retrieves whether DrawLogEntriesPath checkbox should be checked */
@@ -199,6 +206,12 @@ private:
 	/** retrieves whether IgnoreTrivialLogs checkbox should be checked */
 	ESlateCheckBoxState::Type GetIgnoreTrivialLogs() const;
 
+	void OnChangeHistogramLabelLocation(ESlateCheckBoxState::Type NewState);
+	ESlateCheckBoxState::Type GetHistogramLabelLocation() const;
+
+	void OnStickToLastData(ESlateCheckBoxState::Type NewState);
+	ESlateCheckBoxState::Type GetStickToLastData() const;
+	
 	/** handler for toggling log visualizer camera in game view */
 	void OnToggleCamera(ESlateCheckBoxState::Type NewState);
 	/** retrieves whether ToggleCamera toggle button should appear pressed */
@@ -252,15 +265,20 @@ private:
 	float MaxZoom;
 	
 	FOnZoomChanged ZoomChangedNotify;
+	FOnHistogramWindowChanged HistogramWindowChangedNotify;
+	float HistogramPreviewWindow;
 
 	FDebugDrawDelegate DrawingOnCanvasDelegate;
 
 	bool bDrawLogEntriesPath;
 	bool bIgnoreTrivialLogs;
+	bool bShowHistogramLabelsOutside;
+	bool bStickToLastData;
 
 	TWeakObjectPtr<class ALogVisualizerCameraController> CameraController;
 	TArray< TSharedPtr<FLogStatusItem> > StatusItems;
 
+
 	// WIDGETS
 
 	/** Main log list widget */
diff --git a/Engine/Source/Developer/LogVisualizer/Private/SLogsTableRow.cpp b/Engine/Source/Developer/LogVisualizer/Private/SLogsTableRow.cpp
--- a/Engine/Source/Developer/LogVisualizer/Private/SLogsTableRow.cpp
+++ b/Engine/Source/Developer/LogVisualizer/Private/SLogsTableRow.cpp
@@ -50,6 +50,10 @@ TSharedRef<SWidget> SLogsTableRow::GenerateWidgetForColumn(const FName& ColumnNa
 			OwnerVisualizerWidget->OnZoomChanged().Remove(
 				SLogVisualizer::FOnZoomChanged::FDelegate::CreateSP(LogBar.Get(), &SLogBar::SetZoomAndOffset)
 				);
+
+			OwnerVisualizerWidget->OnHistogramWindowChanged().Remove(
+				SLogVisualizer::FOnHistogramWindowChanged::FDelegate::CreateSP(LogBar.Get(), &SLogBar::SetHistogramWindow)
+				);
 		}
 
 		LogBar = SNew( SLogBar )
@@ -64,11 +68,16 @@ TSharedRef<SWidget> SLogsTableRow::GenerateWidgetForColumn(const FName& ColumnNa
 		LogBar->SetZoom( OwnerVisualizerWidget->GetZoom() );
 		LogBar->SetOffset( OwnerVisualizerWidget->GetScrollbarOffset() );
 		LogBar->SetRowIndex(IndexInList);
+		LogBar->SetHistogramWindow(OwnerVisualizerWidget->GetHistogramPreviewWindow() );
 
 		OwnerVisualizerWidget->OnZoomChanged().Add(
 			SLogVisualizer::FOnZoomChanged::FDelegate::CreateSP(LogBar.Get(), &SLogBar::SetZoomAndOffset)
 			);
 
+		OwnerVisualizerWidget->OnHistogramWindowChanged().Add(
+			SLogVisualizer::FOnHistogramWindowChanged::FDelegate::CreateSP(LogBar.Get(), &SLogBar::SetHistogramWindow)
+			);
+
 		return LogBar.ToSharedRef();
 	}
 
diff --git a/Engine/Source/Developer/LogVisualizer/Public/ILogVisualizer.h b/Engine/Source/Developer/LogVisualizer/Public/ILogVisualizer.h
--- a/Engine/Source/Developer/LogVisualizer/Public/ILogVisualizer.h
+++ b/Engine/Source/Developer/LogVisualizer/Public/ILogVisualizer.h
@@ -13,4 +13,7 @@ public:
 
 	/** closes LogVisualizer's UI */
 	virtual void CloseUI(UWorld* InWorld) = 0;
+
+	/** checks if we have LogVisualizer UI opened for given World */
+	virtual bool IsOpenUI(UWorld* InWorld) = 0;
 };
\ No newline at end of file
diff --git a/Engine/Source/Developer/TaskGraph/Public/STimeline.h b/Engine/Source/Developer/TaskGraph/Public/STimeline.h
--- a/Engine/Source/Developer/TaskGraph/Public/STimeline.h
+++ b/Engine/Source/Developer/TaskGraph/Public/STimeline.h
@@ -48,6 +48,10 @@ public:
 		Zoom = FMath::Max(InZoom, 1.0f);
 	}
 
+	float GetZoom()
+	{
+		return Zoom;
+	}
 	/**
 	 * Sets the graph's offset by which all graph bars should be moved
 	 *
@@ -74,13 +78,19 @@ public:
 	 * @param InMin	New min value.
 	 * @param InMin	New max value.
 	 */
-	void SetMinMaxValues( float InMin, float InMax )
+	void SetMinMaxValues(float InMin, float InMax)
 	{
 		MinValue = InMin;
 		MaxValue = InMax;
-	} 
+	}
+
+	void GetMinMaxValues(float &InMin, float &InMax)
+	{
+		InMin = MinValue;
+		InMax = MaxValue;
+	}
 
-	void SetDrawingGeometry( const FGeometry& Geometry )
+	void SetDrawingGeometry(const FGeometry& Geometry)
 	{
 		DrawingGeometry = Geometry;
 	}
diff --git a/Engine/Source/Editor/BlueprintGraph/Private/K2Node_SwitchEnum.cpp b/Engine/Source/Editor/BlueprintGraph/Private/K2Node_SwitchEnum.cpp
--- a/Engine/Source/Editor/BlueprintGraph/Private/K2Node_SwitchEnum.cpp
+++ b/Engine/Source/Editor/BlueprintGraph/Private/K2Node_SwitchEnum.cpp
@@ -35,7 +35,7 @@ void UK2Node_SwitchEnum::SetEnum(UEnum* InEnum)
 				FString const EnumValueName = Enum->GetEnumName(EnumIndex);
 				EnumEntries.Add( FName(*EnumValueName) );
 
-				FString EnumFriendlyName = Enum->GetDisplayNameText(EnumIndex).ToString();
+				FString EnumFriendlyName = Enum->GetEnumText(EnumIndex).ToString();
 				if (EnumFriendlyName.Len() == 0)
 				{
 					EnumFriendlyName = EnumValueName;
diff --git a/Engine/Source/Editor/EnvironmentQueryEditor/Private/DetailCustomizations/EnvQueryTestDetails.cpp b/Engine/Source/Editor/EnvironmentQueryEditor/Private/DetailCustomizations/EnvQueryTestDetails.cpp
--- a/Engine/Source/Editor/EnvironmentQueryEditor/Private/DetailCustomizations/EnvQueryTestDetails.cpp
+++ b/Engine/Source/Editor/EnvironmentQueryEditor/Private/DetailCustomizations/EnvQueryTestDetails.cpp
@@ -51,39 +51,39 @@ void FEnvQueryTestDetails::CustomizeDetails( IDetailLayoutBuilder& DetailLayout
 	BuildScoreClampingTypeValues(true, ClampMinTypeValues);
 	BuildScoreClampingTypeValues(false, ClampMaxTypeValues);
 
-	// dynamic Condition combo
-	IDetailCategoryBuilder& DeprecatedFilterCategory = DetailLayout.EditCategory("Deprecated Filter and Score");
-	IDetailPropertyRow& ConditionRow = DeprecatedFilterCategory.AddProperty(ConditionHandle);
-	ConditionRow.CustomWidget()
-		.NameContent()
-		[
-			ConditionHandle->CreatePropertyNameWidget()
-		]
-		.ValueContent()
-		[
-			SNew(SComboButton)
-			.OnGetMenuContent(this, &FEnvQueryTestDetails::OnGetConditionContent)
-			.ContentPadding(FMargin( 2.0f, 2.0f ))
-			.ButtonContent()
-			[
-				SNew(STextBlock) 
-				.Text(this, &FEnvQueryTestDetails::GetCurrentConditionDesc)
-				.Font(IDetailLayoutBuilder::GetDetailFont())
-			]
-		];
-	ConditionRow.EditCondition(TAttribute<bool>(this, &FEnvQueryTestDetails::AllowWritingToFiltersFromScore), NULL);
-
-	// TODO: Remove!
-	IDetailPropertyRow& DiscardFailedRow = DeprecatedFilterCategory.AddProperty(DetailLayout.GetProperty(GET_MEMBER_NAME_CHECKED(UEnvQueryTest, bDiscardFailedItems)));
-	DiscardFailedRow.EditCondition(TAttribute<bool>(this, &FEnvQueryTestDetails::AllowWritingToFiltersFromScore), NULL);
-
-	// DEPRECATED!  Remove as soon as testing is complete!
-	IDetailPropertyRow& FloatFilterRow = DeprecatedFilterCategory.AddProperty(DetailLayout.GetProperty(GET_MEMBER_NAME_CHECKED(UEnvQueryTest, FloatFilter)));
-	FloatFilterRow.Visibility(TAttribute<EVisibility>::Create(TAttribute<EVisibility>::FGetter::CreateSP(this, &FEnvQueryTestDetails::GetFloatFilterVisibility)));
-	FloatFilterRow.EditCondition(TAttribute<bool>(this, &FEnvQueryTestDetails::AllowWritingToFiltersFromScore), NULL);
-
-	IDetailPropertyRow& WeightModiferRow = DeprecatedFilterCategory.AddProperty(DetailLayout.GetProperty(GET_MEMBER_NAME_CHECKED(UEnvQueryTest, WeightModifier)));
-	WeightModiferRow.EditCondition(TAttribute<bool>(this, &FEnvQueryTestDetails::AllowWritingToFiltersFromScore), NULL);
+// 	// dynamic Condition combo
+// 	IDetailCategoryBuilder& DeprecatedFilterCategory = DetailLayout.EditCategory("Deprecated Filter and Score");
+// 	IDetailPropertyRow& ConditionRow = DeprecatedFilterCategory.AddProperty(ConditionHandle);
+// 	ConditionRow.CustomWidget()
+// 		.NameContent()
+// 		[
+// 			ConditionHandle->CreatePropertyNameWidget()
+// 		]
+// 		.ValueContent()
+// 		[
+// 			SNew(SComboButton)
+// 			.OnGetMenuContent(this, &FEnvQueryTestDetails::OnGetConditionContent)
+// 			.ContentPadding(FMargin( 2.0f, 2.0f ))
+// 			.ButtonContent()
+// 			[
+// 				SNew(STextBlock) 
+// 				.Text(this, &FEnvQueryTestDetails::GetCurrentConditionDesc)
+// 				.Font(IDetailLayoutBuilder::GetDetailFont())
+// 			]
+// 		];
+// 	ConditionRow.EditCondition(TAttribute<bool>(this, &FEnvQueryTestDetails::AllowWritingToFiltersFromScore), NULL);
+// 
+// 	// TODO: Remove!
+// 	IDetailPropertyRow& DiscardFailedRow = DeprecatedFilterCategory.AddProperty(DetailLayout.GetProperty(GET_MEMBER_NAME_CHECKED(UEnvQueryTest, bDiscardFailedItems)));
+// 	DiscardFailedRow.EditCondition(TAttribute<bool>(this, &FEnvQueryTestDetails::AllowWritingToFiltersFromScore), NULL);
+// 
+// 	// DEPRECATED!  Remove as soon as testing is complete!
+// 	IDetailPropertyRow& FloatFilterRow = DeprecatedFilterCategory.AddProperty(DetailLayout.GetProperty(GET_MEMBER_NAME_CHECKED(UEnvQueryTest, FloatFilter)));
+// 	FloatFilterRow.Visibility(TAttribute<EVisibility>::Create(TAttribute<EVisibility>::FGetter::CreateSP(this, &FEnvQueryTestDetails::GetFloatFilterVisibility)));
+// 	FloatFilterRow.EditCondition(TAttribute<bool>(this, &FEnvQueryTestDetails::AllowWritingToFiltersFromScore), NULL);
+// 
+// 	IDetailPropertyRow& WeightModiferRow = DeprecatedFilterCategory.AddProperty(DetailLayout.GetProperty(GET_MEMBER_NAME_CHECKED(UEnvQueryTest, WeightModifier)));
+// 	WeightModiferRow.EditCondition(TAttribute<bool>(this, &FEnvQueryTestDetails::AllowWritingToFiltersFromScore), NULL);
 	
 	IDetailCategoryBuilder& TestCategory = DetailLayout.EditCategory("Test");
 	IDetailPropertyRow& TestPurposeRow = TestCategory.AddProperty(TestPurposeHandle);
@@ -125,15 +125,11 @@ void FEnvQueryTestDetails::CustomizeDetails( IDetailLayoutBuilder& DetailLayout
 	// Scoring
 	IDetailCategoryBuilder& ScoreCategory = DetailLayout.EditCategory("Score");
 
-// 	IDetailPropertyRow& AbsoluteValueOfTestValueRow = ScoreCategory.AddProperty(DetailLayout.GetProperty(GET_MEMBER_NAME_CHECKED(UEnvQueryTest, bUseAbsoluteValueBeforeClamping)));
-// 	AbsoluteValueOfTestValueRow.Visibility(TAttribute<EVisibility>::Create(TAttribute<EVisibility>::FGetter::CreateSP(this, &FEnvQueryTestDetails::GetFloatScoreVisibility)));
-
 	//----------------------------
 	// BEGIN Scoring: Clamping
 	IDetailGroup& ClampingGroup = ScoreCategory.AddGroup("Clamping", FString(TEXT("Clamping")));
 	
 	// Drop-downs for setting type of lower and upper bound normalization
-	// IDetailPropertyRow& ClampMinTypeRow = ScoreCategory.AddProperty(ClampMinTypeHandle);
 	IDetailPropertyRow& ClampMinTypeRow = ClampingGroup.AddPropertyRow(ClampMinTypeHandle.ToSharedRef());
 	ClampMinTypeRow.CustomWidget()
 		.NameContent()
@@ -154,6 +150,16 @@ void FEnvQueryTestDetails::CustomizeDetails( IDetailLayoutBuilder& DetailLayout
 		];
 	ClampMinTypeRow.Visibility(TAttribute<EVisibility>::Create(TAttribute<EVisibility>::FGetter::CreateSP(this, &FEnvQueryTestDetails::GetFloatScoreVisibility)));
 
+	// Lower Bound for normalization of score if specified independently of filtering.
+	IDetailPropertyRow& ScoreClampingMinRow = ClampingGroup.AddPropertyRow(ScoreClampingMinHandle.ToSharedRef());
+	ScoreClampingMinRow.Visibility(TAttribute<EVisibility>::Create(TAttribute<EVisibility>::FGetter::CreateSP(this, &FEnvQueryTestDetails::GetVisibilityOfScoreClampingMinimum)));
+
+	// Lower Bound for scoring when tied to filter minimum.
+	IDetailPropertyRow& FloatFilterMinForClampingRow = ClampingGroup.AddPropertyRow(FloatFilterMinHandle.ToSharedRef());
+	FloatFilterMinForClampingRow.Visibility(TAttribute<EVisibility>::Create(TAttribute<EVisibility>::FGetter::CreateSP(this, &FEnvQueryTestDetails::GetVisibilityOfFilterMinForScoreClamping)));
+	FloatFilterMinForClampingRow.ToolTip(TEXT("See Filter Thresholds under the Filter tab.  Values lower than this (before clamping) cause the item to be thrown out as invalid.  Values are normalized with this value as the minimum, so items with this value will have a normalized score of 0."));
+	FloatFilterMinForClampingRow.EditCondition(TAttribute<bool>(this, &FEnvQueryTestDetails::AllowWritingToFiltersFromScore), NULL);
+
 	IDetailPropertyRow& ClampMaxTypeRow = ClampingGroup.AddPropertyRow(ClampMaxTypeHandle.ToSharedRef());
 	ClampMaxTypeRow.CustomWidget()
 		.NameContent()
@@ -174,16 +180,6 @@ void FEnvQueryTestDetails::CustomizeDetails( IDetailLayoutBuilder& DetailLayout
 		];
 	ClampMaxTypeRow.Visibility(TAttribute<EVisibility>::Create(TAttribute<EVisibility>::FGetter::CreateSP(this, &FEnvQueryTestDetails::GetFloatScoreVisibility)));
 
-	// Lower Bound for normalization of score if specified independently of filtering.
-	IDetailPropertyRow& ScoreClampingMinRow = ClampingGroup.AddPropertyRow(ScoreClampingMinHandle.ToSharedRef());
-	ScoreClampingMinRow.Visibility(TAttribute<EVisibility>::Create(TAttribute<EVisibility>::FGetter::CreateSP(this, &FEnvQueryTestDetails::GetVisibilityOfScoreClampingMinimum)));
-
-	// Lower Bound for scoring when tied to filter minimum.
-	IDetailPropertyRow& FloatFilterMinForClampingRow = ClampingGroup.AddPropertyRow(FloatFilterMinHandle.ToSharedRef());
-	FloatFilterMinForClampingRow.Visibility(TAttribute<EVisibility>::Create(TAttribute<EVisibility>::FGetter::CreateSP(this, &FEnvQueryTestDetails::GetVisibilityOfFilterMinForScoreClamping)));
-	FloatFilterMinForClampingRow.ToolTip(TEXT("See Filter Thresholds under the Filter tab.  Values lower than this (before clamping) cause the item to be thrown out as invalid.  Values are normalized with this value as the minimum, so items with this value will have a normalized score of 0."));
-	FloatFilterMinForClampingRow.EditCondition(TAttribute<bool>(this, &FEnvQueryTestDetails::AllowWritingToFiltersFromScore), NULL);
-
 	// Upper Bound for normalization of score if specified independently of filtering.
 	IDetailPropertyRow& ScoreClampingMaxRow = ClampingGroup.AddPropertyRow(ScoreClampingMaxHandle.ToSharedRef());
 	ScoreClampingMaxRow.Visibility(TAttribute<EVisibility>::Create(TAttribute<EVisibility>::FGetter::CreateSP(this, &FEnvQueryTestDetails::GetVisibilityOfScoreClampingMaximum)));
@@ -197,8 +193,10 @@ void FEnvQueryTestDetails::CustomizeDetails( IDetailLayoutBuilder& DetailLayout
 	// END Scoring: Clamping, continue Scoring
 	//----------------------------
 
-	IDetailPropertyRow& BoolScoreTestRow = FilterCategory.AddProperty(DetailLayout.GetProperty(GET_MEMBER_NAME_CHECKED(UEnvQueryTest, BoolFilter)));
+	IDetailPropertyRow& BoolScoreTestRow = ScoreCategory.AddProperty(DetailLayout.GetProperty(GET_MEMBER_NAME_CHECKED(UEnvQueryTest, BoolFilter)));
 	BoolScoreTestRow.Visibility(TAttribute<EVisibility>::Create(TAttribute<EVisibility>::FGetter::CreateSP(this, &FEnvQueryTestDetails::GetBoolFilterVisibilityForScoring)));
+	BoolScoreTestRow.DisplayName(TEXT("Bool Match"));
+	BoolScoreTestRow.ToolTip(TEXT("Boolean value to match in order to grant score of 'Weight'.  Not matching this value will not change score."));
 
 // 	IDetailPropertyRow& ScoreMirrorNormalizedScoreRow = ScoreCategory.AddProperty(DetailLayout.GetProperty(GET_MEMBER_NAME_CHECKED(UEnvQueryTest, bMirrorNormalizedScore)));
 // 	ScoreMirrorNormalizedScoreRow.Visibility(TAttribute<EVisibility>::Create(TAttribute<EVisibility>::FGetter::CreateSP(this, &FEnvQueryTestDetails::GetFloatScoreVisibility)));
@@ -237,23 +235,11 @@ void FEnvQueryTestDetails::CustomizeDetails( IDetailLayoutBuilder& DetailLayout
 				.Font( IDetailLayoutBuilder::GetDetailFont() )
 			]
 		];
-	ScoreEquationTypeRow.Visibility(TAttribute<EVisibility>::Create(TAttribute<EVisibility>::FGetter::CreateSP(this, &FEnvQueryTestDetails::GetScoreVisibility)));
+	ScoreEquationTypeRow.Visibility(TAttribute<EVisibility>::Create(TAttribute<EVisibility>::FGetter::CreateSP(this, &FEnvQueryTestDetails::GetFloatScoreVisibility)));
 
 	IDetailPropertyRow& ScoreWeightRow = ScoreCategory.AddProperty(DetailLayout.GetProperty(GET_MEMBER_NAME_CHECKED(UEnvQueryTest, Weight)));
 	ScoreWeightRow.Visibility(TAttribute<EVisibility>::Create(TAttribute<EVisibility>::FGetter::CreateSP(this, &FEnvQueryTestDetails::GetScoreVisibility)));
 
-//  	IDetailPropertyRow& ScoreFactorRow = ScoreCategory.AddProperty(DetailLayout.GetProperty(GET_MEMBER_NAME_CHECKED(UEnvQueryTest, ScoringFactor)));
-//  	ScoreFactorRow.Visibility(TAttribute<EVisibility>::Create(TAttribute<EVisibility>::FGetter::CreateSP(this, &FEnvQueryTestDetails::GetScoreFactorVisibility)));
-// 	ScoreConstantRow.DisplayName(FString(TEXT("Weight"))); // Temporarily leaving the name as "Weight" even though I've renamed the property itself.  Ideally 
-
-// 
-// // 	IDetailPropertyRow& ScorePowerRow = ScoreCategory.AddProperty(DetailLayout.GetProperty(GET_MEMBER_NAME_CHECKED(UEnvQueryTest, ScoringPower)));
-// // 	ScorePowerRow.Visibility(TAttribute<EVisibility>::Create(TAttribute<EVisibility>::FGetter::CreateSP(this, &FEnvQueryTestDetails::GetScorePowerVisibility)));
-// 
-// 	IDetailPropertyRow& ScoreConstantRow = ScoreCategory.AddProperty(DetailLayout.GetProperty(GET_MEMBER_NAME_CHECKED(UEnvQueryTest, ScoringFactor)));
-// 	ScoreConstantRow.Visibility(TAttribute<EVisibility>::Create(TAttribute<EVisibility>::FGetter::CreateSP(this, &FEnvQueryTestDetails::GetScoreConstantVisibility)));
-// 	ScoreConstantRow.DisplayName(FString(TEXT("Constant")));
-
 	// scoring & filter function preview
 	IDetailCategoryBuilder& PreviewCategory = DetailLayout.EditCategory("Preview");
 	PreviewCategory.AddCustomRow("Preview").WholeRowWidget
@@ -359,21 +345,6 @@ void FEnvQueryTestDetails::OnScoreEquationChange(int32 Index)
 	ScoreEquationHandle->SetValue(EnumValue);
 }
 
- // TSharedRef<SWidget> OnGetScoringNormMaxTypeContent();
-// FString GetScoringNormMaxTypeDesc() const;
-// 
-// TSharedRef<SWidget> OnGetScoringNormTypeContent();
-// FString GetScoringNormTypeDesc() const;
-
-// 	TEnumAsByte<EEnvQueryTestClamping::Type> ClampMinType;
-// 	TEnumAsByte<EEnvQueryTestClamping::Type> ClampMaxType;
-// 	UEnvQueryTest* MyTestOb = Cast<UEnvQueryTest>(MyTest.Get());
-
-	// 		if ((MyTestOb != NULL) && (MyTestOb->ClampMinType
-
-// TEnumAsByte<EEnvQueryTestClamping::Type> NormalizationType
-// TEnumAsByte<EEnvTestFilterType::Type>
-
 void FEnvQueryTestDetails::BuildScoreClampingTypeValues(bool bBuildMinValues, TArray<FStringIntPair>& ClampTypeValues) const
 {
 	UEnum* ScoringNormalizationEnum = FindObject<UEnum>(ANY_PACKAGE, TEXT("EEnvQueryTestClamping"));
@@ -691,11 +662,8 @@ EVisibility FEnvQueryTestDetails::GetFloatFilterVisibility() const
 {
 	if (IsFiltering())
 	{
-// 		uint8 EnumValue;
-// 		ConditionHandle->GetValue(EnumValue);
 		const UEnvQueryTest* MyTestOb = Cast<const UEnvQueryTest>(MyTest.Get());
 		if (MyTestOb && MyTestOb->GetWorkOnFloatValues())
-		// 		&& (EnumValue == EEnvTestCondition::AtLeast || EnumValue == EEnvTestCondition::UpTo))
 		{
 			return EVisibility::Visible;
 		}
@@ -714,86 +682,6 @@ EVisibility FEnvQueryTestDetails::GetScoreVisibility() const
 	return EVisibility::Collapsed;
 }
 
-// EVisibility FEnvQueryTestDetails::GetScoreFactorVisibility() const
-// {
-// 	if (IsScoring())
-// 	{
-// 		if (ScoreEquationHandle.IsValid())
-// 		{
-// 			uint8 EnumValue;
-// 			ScoreEquationHandle->GetValue(EnumValue);
-// 
-// 			// For now, show Score "Factor" for everything, even constant value.  (Factor * 1)
-// // 			if (EnumValue != EEnvTestScoreEquation::Constant)
-// // 			{
-// // 				return EVisibility::Visible;
-// // 			}
-// 		}
-// 	}
-// 
-// 	return EVisibility::Collapsed;
-// }
-
-// EVisibility FEnvQueryTestDetails::GetScoreConstantVisibility() const
-// {
-// 	if (IsScoring())
-// 	{
-// 		if (ScoreEquationHandle.IsValid())
-// 		{
-// 			uint8 EnumValue;
-// 			ScoreEquationHandle->GetValue(EnumValue);
-// 
-// 			switch (EnumValue)
-// 			{
-// 				case EEnvTestScoreEquation::Constant:
-// // 				case EEnvTestScoreEquation::Parametric:
-// 					return EVisibility::Visible;
-// 
-// 				case EEnvTestScoreEquation::InverseLinear:
-// 				case EEnvTestScoreEquation::Square:
-// 				case EEnvTestScoreEquation::Linear:
-// 					break;
-// 
-// 				default:
-// 					UE_LOG(LogEnvironmentQueryEditor, Error, TEXT("Invalid Enum type in FEnvQueryTestDetails::GetScoreConstantVisibility"));
-// 					break;
-// 			}
-// 		}
-// 	}
-// 
-// 	return EVisibility::Collapsed;
-// }
-
-// EVisibility FEnvQueryTestDetails::GetScorePowerVisibility() const
-// {
-// 	if (IsScoring())
-// 	{
-// 		if (ScoreEquationHandle.IsValid())
-// 		{
-// 			uint8 EnumValue;
-// 			ScoreEquationHandle->GetValue(EnumValue);
-// 
-// 			switch (EnumValue)
-// 			{
-// 				case EEnvTestScoreEquation::Parametric:
-// 					return EVisibility::Visible;
-// 
-// 				case EEnvTestScoreEquation::Constant:
-// 				case EEnvTestScoreEquation::InverseLinear:
-// 				case EEnvTestScoreEquation::Square:
-// 				case EEnvTestScoreEquation::Linear:
-// 					break;
-// 
-// 				default:
-// 					UE_LOG(LogEnvironmentQueryEditor, Error, TEXT("Invalid Enum type in FEnvQueryTestDetails::GetScoreConstantVisibility"));
-// 					break;
-// 			}
-// 		}
-// 	}
-// 
-// 	return EVisibility::Collapsed;
-// }
-
 EVisibility FEnvQueryTestDetails::GetFloatScoreVisibility() const
 {
 	if (IsScoring())
diff --git a/Engine/Source/Editor/EnvironmentQueryEditor/Private/DetailCustomizations/EnvQueryTestDetails.h b/Engine/Source/Editor/EnvironmentQueryEditor/Private/DetailCustomizations/EnvQueryTestDetails.h
--- a/Engine/Source/Editor/EnvironmentQueryEditor/Private/DetailCustomizations/EnvQueryTestDetails.h
+++ b/Engine/Source/Editor/EnvironmentQueryEditor/Private/DetailCustomizations/EnvQueryTestDetails.h
@@ -53,10 +53,6 @@ protected:
 	FString GetScoreEquationInfo() const;
 
  	EVisibility GetScoreVisibility() const;
-// 	EVisibility GetScoreFactorVisibility() const;
-// 	EVisibility GetScoreConstantVisibility() const;
-	// Only needed if we add support for parametric equations
-// 	EVisibility GetScorePowerVisibility() const;
 
 	TSharedRef<SWidget> OnGetFilterTestContent();
 	void BuildFilterTestValues();
diff --git a/Engine/Source/Editor/GameplayTagsEditor/Private/GameplayTagCustomization.cpp b/Engine/Source/Editor/GameplayTagsEditor/Private/GameplayTagCustomization.cpp
new file mode 100644
--- /dev/null
+++ b/Engine/Source/Editor/GameplayTagsEditor/Private/GameplayTagCustomization.cpp
@@ -0,0 +1,156 @@
+// Copyright 1998-2014 Epic Games, Inc. All Rights Reserved.
+
+#include "GameplayTagsEditorModulePrivatePCH.h"
+#include "GameplayTagCustomization.h"
+#include "Editor/PropertyEditor/Public/PropertyEditing.h"
+#include "MainFrame.h"
+#include "SGameplayTagWidget.h"
+#include "GameplayTagContainer.h"
+#include "ScopedTransaction.h"
+
+#define LOCTEXT_NAMESPACE "GameplayTagCustomization"
+
+void FGameplayTagCustomization::CustomizeHeader(TSharedRef<class IPropertyHandle> InStructPropertyHandle, class FDetailWidgetRow& HeaderRow, IPropertyTypeCustomizationUtils& StructCustomizationUtils)
+{
+	TagContainer = MakeShareable(new FGameplayTagContainer);
+	StructPropertyHandle = InStructPropertyHandle;
+
+	BuildEditableContainerList();
+
+	HeaderRow
+	.NameContent()
+	[
+		StructPropertyHandle->CreatePropertyNameWidget()
+	]
+	.ValueContent()
+	.MaxDesiredWidth(512)
+	[
+		SNew(SHorizontalBox)
+		+ SHorizontalBox::Slot()
+		.AutoWidth()
+		[
+			SAssignNew(ComboButton, SComboButton)
+			.OnGetMenuContent(this, &FGameplayTagCustomization::GetListContent)
+			.ContentPadding(FMargin(2.0f, 2.0f))
+			.MenuPlacement(MenuPlacement_BelowAnchor)
+			.ButtonContent()
+			[
+				SNew(STextBlock)
+				.Text(LOCTEXT("GameplayTagCustomization_Edit", "Edit"))
+			]
+		]
+		+ SHorizontalBox::Slot()
+		.AutoWidth()
+		[
+			SNew(SBorder)
+			.Padding(4.0f)
+			[
+				SNew(STextBlock)
+				.Text(this, &FGameplayTagCustomization::SelectedTag)
+			]
+		]
+	];
+
+	GEditor->RegisterForUndo(this);
+}
+
+TSharedRef<SWidget> FGameplayTagCustomization::GetListContent()
+{
+	BuildEditableContainerList();
+	
+	FString Categories;
+	if (StructPropertyHandle->GetProperty()->HasMetaData(TEXT("Categories")))
+	{
+		Categories = StructPropertyHandle->GetProperty()->GetMetaData(TEXT("Categories"));
+	}
+
+	bool bReadOnly = StructPropertyHandle->GetProperty()->HasAnyPropertyFlags(CPF_EditConst);
+
+	return SNew(SVerticalBox)
+		+ SVerticalBox::Slot()
+		.AutoHeight()
+		.MaxHeight(400)
+		[
+			SNew(SGameplayTagWidget, EditableContainers)
+			.Filter(Categories)
+			.ReadOnly(bReadOnly)
+			.TagContainerName(StructPropertyHandle->GetPropertyDisplayName())
+			.MultiSelect(false)
+			.OnTagChanged(this, &FGameplayTagCustomization::OnTagChanged)
+		];
+}
+
+void FGameplayTagCustomization::OnTagChanged()
+{
+	TagName = TEXT("");
+	if (StructPropertyHandle.IsValid() && EditableContainers.Num() > 0)
+	{
+		TArray<void*> RawStructData;
+		StructPropertyHandle->AccessRawData(RawStructData);
+		if (RawStructData.Num() > 0)
+		{
+			FGameplayTag* Tag = (FGameplayTag*)(RawStructData[0]);
+
+			// Update Tag from the one selected from list
+			FGameplayTagContainer* Container = EditableContainers[0].TagContainer;
+			if (Container)
+			{
+				for (auto It = Container->CreateConstIterator(); It; ++It)
+				{
+					*Tag = *It;
+					TagName = It->ToString();
+				}
+			}
+		}
+	}
+}
+
+void FGameplayTagCustomization::PostUndo(bool bSuccess)
+{
+	if( bSuccess )
+	{
+		OnTagChanged();
+	}
+}
+
+void FGameplayTagCustomization::PostRedo(bool bSuccess)
+{
+	if( bSuccess )
+	{
+		OnTagChanged();
+	}
+}
+
+FGameplayTagCustomization::~FGameplayTagCustomization()
+{
+	GEditor->UnregisterForUndo(this);
+}
+
+void FGameplayTagCustomization::BuildEditableContainerList()
+{
+	EditableContainers.Empty();
+
+	if(StructPropertyHandle.IsValid())
+	{
+		TArray<void*> RawStructData;
+		StructPropertyHandle->AccessRawData(RawStructData);
+
+		TArray<UObject*> OuterObjects;
+		StructPropertyHandle->GetOuterObjects(OuterObjects);
+
+		FGameplayTag* Tag = (FGameplayTag*)(RawStructData[0]);
+		if (Tag->IsValid())
+		{
+			TagName = Tag->ToString();
+			TagContainer->AddTag(*Tag);
+		}
+		EditableContainers.Add(SGameplayTagWidget::FEditableGameplayTagContainerDatum(OuterObjects[0], TagContainer.Get()));
+	}
+}
+
+FText FGameplayTagCustomization::SelectedTag() const
+{
+	return FText::FromString(*TagName);
+}
+
+#undef LOCTEXT_NAMESPACE
diff --git a/Engine/Source/Editor/GameplayTagsEditor/Private/GameplayTagCustomization.h b/Engine/Source/Editor/GameplayTagsEditor/Private/GameplayTagCustomization.h
new file mode 100644
--- /dev/null
+++ b/Engine/Source/Editor/GameplayTagsEditor/Private/GameplayTagCustomization.h
@@ -0,0 +1,59 @@
+// Copyright 1998-2014 Epic Games, Inc. All Rights Reserved.
+
+#pragma once
+
+#include "Editor/PropertyEditor/Public/PropertyEditorModule.h"
+#include "SGameplayTagWidget.h"
+
+/** Customization for the gameplay tag struct */
+class FGameplayTagCustomization : public IPropertyTypeCustomization, public FEditorUndoClient
+{
+public:
+	static TSharedRef<IPropertyTypeCustomization> MakeInstance()
+	{
+		return MakeShareable(new FGameplayTagCustomization);
+	}
+
+	~FGameplayTagCustomization();
+
+	/** Overridden to show an edit button to launch the gameplay tag editor */
+	virtual void CustomizeHeader(TSharedRef<class IPropertyHandle> StructPropertyHandle, class FDetailWidgetRow& HeaderRow, IPropertyTypeCustomizationUtils& StructCustomizationUtils) override;
+	
+	/** Overridden to do nothing */
+	virtual void CustomizeChildren(TSharedRef<IPropertyHandle> InStructPropertyHandle, class IDetailChildrenBuilder& ChildBuilder, IPropertyTypeCustomizationUtils& StructCustomizationUtils) override {}
+
+	// Begin FEditorUndoClient Interface
+	virtual void PostUndo(bool bSuccess) override;
+	virtual void PostRedo(bool bSuccess) override;
+	// End FEditorUndoClient Interface
+
+private:
+
+	/** Updates the selected tag*/
+	void OnTagChanged();
+
+	/** Build Editable Container */
+	void BuildEditableContainerList();
+
+	/** Callback function to create content for the combo button. */
+	TSharedRef<SWidget> GetListContent();
+
+	/** Returns Tag name currently selected*/
+	FText SelectedTag() const;
+
+	/** Combo Button for the drop down list. */
+	TSharedPtr<SComboButton> ComboButton;
+
+	/** Cached property handle */
+	TSharedPtr<IPropertyHandle> StructPropertyHandle;
+
+	/** Tag Container used for the GameplayTagWidget. */
+	TSharedPtr<FGameplayTagContainer> TagContainer;
+
+	/** Editable Container for holding our tag */
+	TArray<SGameplayTagWidget::FEditableGameplayTagContainerDatum> EditableContainers;
+
+	/** Tag name selected*/
+	FString TagName;
+};
+
diff --git a/Engine/Source/Editor/GameplayTagsEditor/Private/GameplayTagsEditorModule.cpp b/Engine/Source/Editor/GameplayTagsEditor/Private/GameplayTagsEditorModule.cpp
--- a/Engine/Source/Editor/GameplayTagsEditor/Private/GameplayTagsEditorModule.cpp
+++ b/Engine/Source/Editor/GameplayTagsEditor/Private/GameplayTagsEditorModule.cpp
@@ -3,8 +3,7 @@
 #include "GameplayTagsEditorModulePrivatePCH.h"
 #include "GameplayTagsGraphPanelPinFactory.h"
 #include "GameplayTagContainerCustomization.h"
-
-
+#include "GameplayTagCustomization.h"
 
 class FGameplayTagsEditorModule : public IGameplayTagsEditorModule
 {
@@ -21,6 +20,7 @@ void FGameplayTagsEditorModule::StartupModule()
 	// Register the details customizer
 	FPropertyEditorModule& PropertyModule = FModuleManager::LoadModuleChecked<FPropertyEditorModule>("PropertyEditor");
 	PropertyModule.RegisterCustomPropertyTypeLayout("GameplayTagContainer", FOnGetPropertyTypeCustomizationInstance::CreateStatic(&FGameplayTagContainerCustomization::MakeInstance));
+	PropertyModule.RegisterCustomPropertyTypeLayout("GameplayTag", FOnGetPropertyTypeCustomizationInstance::CreateStatic(&FGameplayTagCustomization::MakeInstance));
 
 	TSharedPtr<FGameplayTagsGraphPanelPinFactory> GameplayTagsGraphPanelPinFactory = MakeShareable( new FGameplayTagsGraphPanelPinFactory() );
 	FEdGraphUtilities::RegisterVisualPinFactory(GameplayTagsGraphPanelPinFactory);
diff --git a/Engine/Source/Editor/LevelEditor/Public/ILevelViewport.h b/Engine/Source/Editor/LevelEditor/Public/ILevelViewport.h
--- a/Engine/Source/Editor/LevelEditor/Public/ILevelViewport.h
+++ b/Engine/Source/Editor/LevelEditor/Public/ILevelViewport.h
@@ -90,6 +90,9 @@ public:
 	/** Removes a widget that was previously overlaid on to this viewport */
 	virtual void RemoveOverlayWidget(TSharedRef<SWidget> OverlaidWidget) = 0;
 
-	/** Returns the viewport as a widget */
-	virtual TSharedRef< const SWidget> AsWidget() = 0;
+	/** Returns the SLevelViewport widget. This is NOT the actual SViewport widget, just a wrapper that contains one. */
+	virtual TSharedRef< const SWidget > AsWidget() = 0;
+
+	/** Returns the SViewport widget contained within the SLevelViewport. This is the actual viewport widget that renders the level. */
+	virtual TWeakPtr< class SViewport > GetViewportWidget() = 0;
 };
\ No newline at end of file
diff --git a/Engine/Source/Editor/LevelEditor/Public/SLevelViewport.h b/Engine/Source/Editor/LevelEditor/Public/SLevelViewport.h
--- a/Engine/Source/Editor/LevelEditor/Public/SLevelViewport.h
+++ b/Engine/Source/Editor/LevelEditor/Public/SLevelViewport.h
@@ -88,6 +88,8 @@ public:
 	virtual bool HasPlayInEditorViewport() const override; 
 	virtual FViewport* GetActiveViewport() override;
 	virtual TSharedRef< const SWidget> AsWidget() override { return AsShared(); }
+	virtual TWeakPtr< SViewport > GetViewportWidget() override { return ViewportWidget; }
+
 
 	/** SEditorViewport Interface */
 	virtual void OnFocusViewportToSelection() override;
diff --git a/Engine/Source/Editor/UnrealEd/Private/Editor.cpp b/Engine/Source/Editor/UnrealEd/Private/Editor.cpp
--- a/Engine/Source/Editor/UnrealEd/Private/Editor.cpp
+++ b/Engine/Source/Editor/UnrealEd/Private/Editor.cpp
@@ -47,6 +47,7 @@
 #include "Developer/DirectoryWatcher/Public/DirectoryWatcherModule.h"
 
 #include "Runtime/Engine/Public/Slate/SceneViewport.h"
+#include "Editor/LevelEditor/Public/ILevelViewport.h"
 
 #include "ComponentReregisterContext.h"
 #include "EngineModule.h"
@@ -4495,6 +4496,12 @@ TSharedPtr<SViewport> UEditorEngine::GetGameViewportWidget() const
 		{
 			return It.Value().SlatePlayInEditorWindowViewport->GetViewportWidget().Pin();
 		}
+
+		TSharedPtr<ILevelViewport> DestinationLevelViewport = It.Value().DestinationSlateViewport.Pin();
+		if (DestinationLevelViewport.IsValid())
+		{
+			return DestinationLevelViewport->GetViewportWidget().Pin();
+		}
 	}
 
 	/*
diff --git a/Engine/Source/Runtime/AIModule/Classes/AITypes.h b/Engine/Source/Runtime/AIModule/Classes/AITypes.h
--- a/Engine/Source/Runtime/AIModule/Classes/AITypes.h
+++ b/Engine/Source/Runtime/AIModule/Classes/AITypes.h
@@ -111,6 +111,7 @@ namespace EPawnActionEventType
 	enum Type
 	{
 		Invalid,
+		FailedToStart,
 		FinishedAborting,
 		FinishedExecution,
 		Push,
diff --git a/Engine/Source/Runtime/AIModule/Classes/BrainComponent.h b/Engine/Source/Runtime/AIModule/Classes/BrainComponent.h
--- a/Engine/Source/Runtime/AIModule/Classes/BrainComponent.h
+++ b/Engine/Source/Runtime/AIModule/Classes/BrainComponent.h
@@ -31,12 +31,20 @@ struct AIMODULE_API FAIMessage
 	/** message param: status */
 	TEnumAsByte<EStatus> Status;
 
-	FAIMessage() : MessageName(NAME_None), Sender(NULL), RequestID(0), Status(FAIMessage::Success) {}
-	FAIMessage(FName InMessage, UObject* InSender) : MessageName(InMessage), Sender(InSender), RequestID(0), Status(FAIMessage::Success) {}
-	FAIMessage(FName InMessage, UObject* InSender, FAIRequestID InID, EStatus InStatus) : MessageName(InMessage), Sender(InSender), RequestID(InID), Status(InStatus) {}
-	FAIMessage(FName InMessage, UObject* InSender, FAIRequestID InID, bool bSuccess) : MessageName(InMessage), Sender(InSender), RequestID(InID), Status(bSuccess ? Success : Failure) {}
-	FAIMessage(FName InMessage, UObject* InSender, EStatus InStatus) : MessageName(InMessage), Sender(InSender), RequestID(0), Status(InStatus) {}
-	FAIMessage(FName InMessage, UObject* InSender, bool bSuccess) : MessageName(InMessage), Sender(InSender), RequestID(0), Status(bSuccess ? Success : Failure) {}
+	/** message param: custom flags */
+	uint8 MessageFlags;
+
+	FAIMessage() : MessageName(NAME_None), Sender(NULL), RequestID(0), Status(FAIMessage::Success), MessageFlags(0) {}
+	FAIMessage(FName InMessage, UObject* InSender) : MessageName(InMessage), Sender(InSender), RequestID(0), Status(FAIMessage::Success), MessageFlags(0) {}
+	FAIMessage(FName InMessage, UObject* InSender, FAIRequestID InID, EStatus InStatus) : MessageName(InMessage), Sender(InSender), RequestID(InID), Status(InStatus), MessageFlags(0) {}
+	FAIMessage(FName InMessage, UObject* InSender, FAIRequestID InID, bool bSuccess) : MessageName(InMessage), Sender(InSender), RequestID(InID), Status(bSuccess ? Success : Failure), MessageFlags(0) {}
+	FAIMessage(FName InMessage, UObject* InSender, EStatus InStatus) : MessageName(InMessage), Sender(InSender), RequestID(0), Status(InStatus), MessageFlags(0) {}
+	FAIMessage(FName InMessage, UObject* InSender, bool bSuccess) : MessageName(InMessage), Sender(InSender), RequestID(0), Status(bSuccess ? Success : Failure), MessageFlags(0) {}
+
+	void SetFlags(uint8 Flags) { MessageFlags = Flags; }
+	void SetFlag(uint8 Flag) { MessageFlags |= Flag; }
+	void ClearFlag(uint8 Flag) { MessageFlags &= ~Flag; }
+	bool HasFlag(uint8 Flag) const { return (MessageFlags & Flag) != 0; }
 
 	static void Send(class AController* Controller, const FAIMessage& Message);
 	static void Send(class APawn* Pawn, const FAIMessage& Message);
diff --git a/Engine/Source/Runtime/AIModule/Classes/EnvironmentQuery/EnvQueryTest.h b/Engine/Source/Runtime/AIModule/Classes/EnvironmentQuery/EnvQueryTest.h
--- a/Engine/Source/Runtime/AIModule/Classes/EnvironmentQuery/EnvQueryTest.h
+++ b/Engine/Source/Runtime/AIModule/Classes/EnvironmentQuery/EnvQueryTest.h
@@ -26,8 +26,9 @@ class AIMODULE_API UEnvQueryTest : public UObject
 		Meta=(ToolTip="Boolean value that must be matched in order for scoring to occur or filtering test to pass."))
 	FEnvBoolParam BoolFilter;
 
+	// TODO: REMOVE, Deprecated!
 	/** filter for test value; currently this is the ONLY value, because we can only filter in one direction. */
-	UPROPERTY(EditDefaultsOnly, Category=Filter)
+	UPROPERTY()
 	FEnvFloatParam FloatFilter;
 
 	UPROPERTY(EditDefaultsOnly, Category=Filter,
@@ -44,12 +45,12 @@ class AIMODULE_API UEnvQueryTest : public UObject
 
 	// TODO: REMOVE, Deprecated!
 	/** Condition for discarding item */
-	UPROPERTY(EditDefaultsOnly, Category=DeprecatedFilt)
+	UPROPERTY()
 	TEnumAsByte<EEnvTestCondition::Type> Condition;
 
 	// TODO: REMOVE, Deprecated!
 	/** Weight modifier */
-	UPROPERTY(EditDefaultsOnly, Category=Weight)
+	UPROPERTY()
 	TEnumAsByte<EEnvTestWeight::Type> WeightModifier;
 
 	/** Cost of test */
@@ -101,7 +102,7 @@ private:
 public:
 	// TODO: REMOVE, Deprecated!
 	/** If set, items not passing filter condition will be removed from query, otherwise they will receive score 0 */
-	UPROPERTY(EditDefaultsOnly, Category=DeprecatedFilt)
+	UPROPERTY()
 	uint32 bDiscardFailedItems : 1;
 
 	FExecuteTestSignature ExecuteDelegate;
diff --git a/Engine/Source/Runtime/AIModule/Classes/Navigation/CrowdFollowingComponent.h b/Engine/Source/Runtime/AIModule/Classes/Navigation/CrowdFollowingComponent.h
--- a/Engine/Source/Runtime/AIModule/Classes/Navigation/CrowdFollowingComponent.h
+++ b/Engine/Source/Runtime/AIModule/Classes/Navigation/CrowdFollowingComponent.h
@@ -34,7 +34,7 @@ class AIMODULE_API UCrowdFollowingComponent : public UPathFollowingComponent, pu
 	// PathFollowingComponent BEGIN
 	virtual void Initialize() override;
 	virtual void Cleanup() override;
-	virtual void AbortMove(const FString& Reason, FAIRequestID RequestID = FAIRequestID::CurrentRequest, bool bResetVelocity = true, bool bSilent = false) override;
+	virtual void AbortMove(const FString& Reason, FAIRequestID RequestID = FAIRequestID::CurrentRequest, bool bResetVelocity = true, bool bSilent = false, uint8 MessageFlags = 0) override;
 	virtual void PauseMove(FAIRequestID RequestID = FAIRequestID::CurrentRequest, bool bResetVelocity = true) override;
 	virtual void ResumeMove(FAIRequestID RequestID = FAIRequestID::CurrentRequest) override;
 	virtual FVector GetMoveFocus(bool bAllowStrafe) const override;
@@ -79,6 +79,9 @@ class AIMODULE_API UCrowdFollowingComponent : public UPathFollowingComponent, pu
 	FORCEINLINE float GetCrowdCollisionQueryRange() const { return CollisionQueryRange; }
 	FORCEINLINE float GetCrowdPathOptimizationRange() const { return PathOptimizationRange; }
 	FORCEINLINE ECrowdAvoidanceQuality::Type GetCrowdAvoidanceQuality() const { return AvoidanceQuality; }
+	FORCEINLINE int32 GetAvoidanceGroup() const { return AvoidanceGroup.Packed; }
+	FORCEINLINE int32 GetGroupsToAvoid() const { return GroupsToAvoid.Packed; }
+	FORCEINLINE int32 GetGroupsToIgnore() const { return GroupsToIgnore.Packed; }
 
 	virtual void GetDebugStringTokens(TArray<FString>& Tokens, TArray<EPathFollowingDebugTokens::Type>& Flags) const;
 #if ENABLE_VISUAL_LOG
@@ -90,6 +93,18 @@ protected:
 	UPROPERTY(transient)
 	class UCharacterMovementComponent* CharacterMovement;
 
+	/** Group mask for this agent */
+	UPROPERTY(Category = "Avoidance", EditAnywhere, BlueprintReadOnly, AdvancedDisplay)
+	FNavAvoidanceMask AvoidanceGroup;
+
+	/** Will avoid other agents if they are in one of specified groups */
+	UPROPERTY(Category = "Avoidance", EditAnywhere, BlueprintReadOnly, AdvancedDisplay)
+	FNavAvoidanceMask GroupsToAvoid;
+
+	/** Will NOT avoid other agents if they are in one of specified groups, higher priority than GroupsToAvoid */
+	UPROPERTY(Category = "Avoidance", EditAnywhere, BlueprintReadOnly, AdvancedDisplay)
+	FNavAvoidanceMask GroupsToIgnore;
+
 	/** if set, velocity will be updated even if agent is falling */
 	uint32 bAffectFallingVelocity : 1;
 
diff --git a/Engine/Source/Runtime/AIModule/Classes/Navigation/PathFollowingComponent.h b/Engine/Source/Runtime/AIModule/Classes/Navigation/PathFollowingComponent.h
--- a/Engine/Source/Runtime/AIModule/Classes/Navigation/PathFollowingComponent.h
+++ b/Engine/Source/Runtime/AIModule/Classes/Navigation/PathFollowingComponent.h
@@ -71,6 +71,15 @@ namespace EPathFollowingDebugTokens
 	};
 }
 
+namespace EPathFollowingMessage
+{
+	enum Type
+	{
+		NoPath,			// abort reason for AI message
+		OtherRequest,	// abort reason for AI message
+	};
+}
+
 UCLASS(config=Engine)
 class AIMODULE_API UPathFollowingComponent : public UActorComponent, public IAIResourceInterface
 {
@@ -118,7 +127,7 @@ class AIMODULE_API UPathFollowingComponent : public UActorComponent, public IAIR
 	 *  @param RequestID - request to abort, 0 = current
 	 *  @param bResetVelocity - try to stop movement component
 	 *  @param bSilent - finish with Skipped result instead of Aborted */
-	virtual void AbortMove(const FString& Reason, FAIRequestID RequestID = FAIRequestID::CurrentRequest, bool bResetVelocity = true, bool bSilent = false);
+	virtual void AbortMove(const FString& Reason, FAIRequestID RequestID = FAIRequestID::CurrentRequest, bool bResetVelocity = true, bool bSilent = false, uint8 MessageFlags = 0);
 
 	/** pause path following
 	*  @param RequestID - request to pause, FAIRequestID::CurrentRequest means pause current request, regardless of its ID */
diff --git a/Engine/Source/Runtime/AIModule/Private/Navigation/CrowdFollowingComponent.cpp b/Engine/Source/Runtime/AIModule/Private/Navigation/CrowdFollowingComponent.cpp
--- a/Engine/Source/Runtime/AIModule/Private/Navigation/CrowdFollowingComponent.cpp
+++ b/Engine/Source/Runtime/AIModule/Private/Navigation/CrowdFollowingComponent.cpp
@@ -28,6 +28,10 @@ UCrowdFollowingComponent::UCrowdFollowingComponent(const class FPostConstructIni
 	CollisionQueryRange = 400.0f;		// approx: radius * 12.0f
 	PathOptimizationRange = 1000.0f;	// approx: radius * 30.0f
 	AvoidanceQuality = ECrowdAvoidanceQuality::Low;
+
+	AvoidanceGroup.SetFlagsDirectly(1);
+	GroupsToAvoid.SetFlagsDirectly(MAX_uint32);
+	GroupsToIgnore.SetFlagsDirectly(0);
 }
 
 FVector UCrowdFollowingComponent::GetCrowdAgentLocation() const
@@ -294,7 +298,7 @@ void UCrowdFollowingComponent::Cleanup()
 	}
 }
 
-void UCrowdFollowingComponent::AbortMove(const FString& Reason, FAIRequestID RequestID, bool bResetVelocity, bool bSilent)
+void UCrowdFollowingComponent::AbortMove(const FString& Reason, FAIRequestID RequestID, bool bResetVelocity, bool bSilent, uint8 MessageFlags)
 {
 	if (bEnableCrowdSimulation && (Status != EPathFollowingStatus::Idle) && RequestID.IsEquivalent(GetCurrentRequestId()))
 	{
@@ -305,7 +309,7 @@ void UCrowdFollowingComponent::AbortMove(const FString& Reason, FAIRequestID Req
 		}
 	}
 
-	Super::AbortMove(Reason, RequestID, bResetVelocity, bSilent);
+	Super::AbortMove(Reason, RequestID, bResetVelocity, bSilent, MessageFlags);
 }
 
 void UCrowdFollowingComponent::PauseMove(FAIRequestID RequestID, bool bResetVelocity)
@@ -564,7 +568,7 @@ void UCrowdFollowingComponent::UpdatePathSegment()
 
 	if (!Path.IsValid() || MovementComp == NULL)
 	{
-		AbortMove(TEXT("no path"));
+		AbortMove(TEXT("no path"), FAIRequestID::CurrentRequest, true, false, EPathFollowingMessage::NoPath);
 		return;
 	}
 
@@ -572,7 +576,7 @@ void UCrowdFollowingComponent::UpdatePathSegment()
 	{
 		if (NavComp == NULL || !NavComp->IsWaitingForRepath())
 		{
-			AbortMove(TEXT("no path"));
+			AbortMove(TEXT("no path"), FAIRequestID::CurrentRequest, true, false, EPathFollowingMessage::NoPath);
 		}
 		return;
 	}
diff --git a/Engine/Source/Runtime/AIModule/Private/Navigation/CrowdManager.cpp b/Engine/Source/Runtime/AIModule/Private/Navigation/CrowdManager.cpp
--- a/Engine/Source/Runtime/AIModule/Private/Navigation/CrowdManager.cpp
+++ b/Engine/Source/Runtime/AIModule/Private/Navigation/CrowdManager.cpp
@@ -649,6 +649,15 @@ void UCrowdManager::GetAgentParams(const ICrowdAgentInterface* Agent, struct dtC
 			(CrowdComponent->IsCrowdOptimizeVisibilityEnabled() ? (DT_CROWD_OPTIMIZE_VIS | DT_CROWD_OPTIMIZE_VIS_MULTI) : 0) |
 			(CrowdComponent->IsCrowdOptimizeTopologyEnabled() ? DT_CROWD_OPTIMIZE_TOPO : 0) |
 			(CrowdComponent->IsCrowdPathOffsetEnabled() ? DT_CROWD_OFFSET_PATH : 0);
+
+		AgentParams->avoidanceGroup = CrowdComponent->GetAvoidanceGroup();
+		AgentParams->groupsToAvoid = CrowdComponent->GetGroupsToAvoid();
+		AgentParams->groupsToIgnore = CrowdComponent->GetGroupsToIgnore();
+	}
+	else
+	{
+		AgentParams->avoidanceGroup = 1;
+		AgentParams->groupsToAvoid = MAX_uint32;
 	}
 }
 
diff --git a/Engine/Source/Runtime/AIModule/Private/Navigation/PathFollowingComponent.cpp b/Engine/Source/Runtime/AIModule/Private/Navigation/PathFollowingComponent.cpp
--- a/Engine/Source/Runtime/AIModule/Private/Navigation/PathFollowingComponent.cpp
+++ b/Engine/Source/Runtime/AIModule/Private/Navigation/PathFollowingComponent.cpp
@@ -157,7 +157,7 @@ FAIRequestID UPathFollowingComponent::RequestMove(FNavPathSharedPtr InPath, FReq
 		{
 			const bool bResetVelocity = false;
 			const bool bFinishAsSkipped = true;
-			AbortMove(TEXT("new request"), PrevMoveId, bResetVelocity, bFinishAsSkipped);
+			AbortMove(TEXT("new request"), PrevMoveId, bResetVelocity, bFinishAsSkipped, EPathFollowingMessage::OtherRequest);
 		}
 		
 		Reset();
@@ -235,7 +235,7 @@ bool UPathFollowingComponent::UpdateMove(FNavPathSharedPtr InPath, FAIRequestID
 	return true;
 }
 
-void UPathFollowingComponent::AbortMove(const FString& Reason, FAIRequestID RequestID, bool bResetVelocity, bool bSilent)
+void UPathFollowingComponent::AbortMove(const FString& Reason, FAIRequestID RequestID, bool bResetVelocity, bool bSilent, uint8 MessageFlags)
 {
 	UE_VLOG(GetOwner(), LogPathFollowing, Log, TEXT("AbortMove: Reason(%s) RequestID(%u)"), *Reason, RequestID);
 
@@ -270,7 +270,11 @@ void UPathFollowingComponent::AbortMove(const FString& Reason, FAIRequestID Requ
 		// notify observers after state was reset (they can request another move)
 		SavedReqFinished.ExecuteIfBound(FinishResult);
 		OnMoveFinished.Broadcast(AbortedMoveId, FinishResult);
-		FAIMessage::Send(Cast<AController>(GetOwner()), FAIMessage(UBrainComponent::AIMessage_MoveFinished, this, AbortedMoveId, FAIMessage::Failure));
+
+		FAIMessage Msg(UBrainComponent::AIMessage_MoveFinished, this, AbortedMoveId, FAIMessage::Failure);
+		Msg.SetFlag(MessageFlags);
+
+		FAIMessage::Send(Cast<AController>(GetOwner()), Msg);
 	}
 	else
 	{
@@ -719,7 +723,7 @@ void UPathFollowingComponent::UpdatePathSegment()
 {
 	if (!Path.IsValid() || MovementComp == NULL)
 	{
-		AbortMove(TEXT("no path"));
+		AbortMove(TEXT("no path"), FAIRequestID::CurrentRequest, true, false, EPathFollowingMessage::NoPath);
 		return;
 	}
 
@@ -727,7 +731,7 @@ void UPathFollowingComponent::UpdatePathSegment()
 	{
 		if (NavComp == NULL || !NavComp->IsWaitingForRepath())
 		{
-			AbortMove(TEXT("no path"));
+			AbortMove(TEXT("no path"), FAIRequestID::CurrentRequest, true, false, EPathFollowingMessage::NoPath);
 		}
 		return;
 	}
diff --git a/Engine/Source/Runtime/Engine/Classes/AI/Navigation/NavigationAvoidanceTypes.h b/Engine/Source/Runtime/Engine/Classes/AI/Navigation/NavigationAvoidanceTypes.h
new file mode 100644
--- /dev/null
+++ b/Engine/Source/Runtime/Engine/Classes/AI/Navigation/NavigationAvoidanceTypes.h
@@ -0,0 +1,111 @@
+// Copyright 1998-2014 Epic Games, Inc. All Rights Reserved.
+
+#pragma once
+#include "NavigationAvoidanceTypes.generated.h"
+
+USTRUCT()
+struct FNavAvoidanceMask
+{
+	GENERATED_USTRUCT_BODY()
+
+#if CPP
+	union
+	{
+		struct
+		{
+#endif
+			UPROPERTY(Category = "Group", EditAnywhere, BlueprintReadOnly)
+			uint32 bGroup0 : 1;
+			UPROPERTY(Category = "Group", EditAnywhere, BlueprintReadOnly)
+			uint32 bGroup1 : 1;
+			UPROPERTY(Category = "Group", EditAnywhere, BlueprintReadOnly)
+			uint32 bGroup2 : 1;
+			UPROPERTY(Category = "Group", EditAnywhere, BlueprintReadOnly)
+			uint32 bGroup3 : 1;
+			UPROPERTY(Category = "Group", EditAnywhere, BlueprintReadOnly)
+			uint32 bGroup4 : 1;
+			UPROPERTY(Category = "Group", EditAnywhere, BlueprintReadOnly)
+			uint32 bGroup5 : 1;
+			UPROPERTY(Category = "Group", EditAnywhere, BlueprintReadOnly)
+			uint32 bGroup6 : 1;
+			UPROPERTY(Category = "Group", EditAnywhere, BlueprintReadOnly)
+			uint32 bGroup7 : 1;
+			UPROPERTY(Category = "Group", EditAnywhere, BlueprintReadOnly)
+			uint32 bGroup8 : 1;
+			UPROPERTY(Category = "Group", EditAnywhere, BlueprintReadOnly)
+			uint32 bGroup9 : 1;
+			UPROPERTY(Category = "Group", EditAnywhere, BlueprintReadOnly)
+			uint32 bGroup10 : 1;
+			UPROPERTY(Category = "Group", EditAnywhere, BlueprintReadOnly)
+			uint32 bGroup11 : 1;
+			UPROPERTY(Category = "Group", EditAnywhere, BlueprintReadOnly)
+			uint32 bGroup12 : 1;
+			UPROPERTY(Category = "Group", EditAnywhere, BlueprintReadOnly)
+			uint32 bGroup13 : 1;
+			UPROPERTY(Category = "Group", EditAnywhere, BlueprintReadOnly)
+			uint32 bGroup14 : 1;
+			UPROPERTY(Category = "Group", EditAnywhere, BlueprintReadOnly)
+			uint32 bGroup15 : 1;
+			UPROPERTY(Category = "Group", EditAnywhere, BlueprintReadOnly)
+			uint32 bGroup16 : 1;
+			UPROPERTY(Category = "Group", EditAnywhere, BlueprintReadOnly)
+			uint32 bGroup17 : 1;
+			UPROPERTY(Category = "Group", EditAnywhere, BlueprintReadOnly)
+			uint32 bGroup18 : 1;
+			UPROPERTY(Category = "Group", EditAnywhere, BlueprintReadOnly)
+			uint32 bGroup19 : 1;
+			UPROPERTY(Category = "Group", EditAnywhere, BlueprintReadOnly)
+			uint32 bGroup20 : 1;
+			UPROPERTY(Category = "Group", EditAnywhere, BlueprintReadOnly)
+			uint32 bGroup21 : 1;
+			UPROPERTY(Category = "Group", EditAnywhere, BlueprintReadOnly)
+			uint32 bGroup22 : 1;
+			UPROPERTY(Category = "Group", EditAnywhere, BlueprintReadOnly)
+			uint32 bGroup23 : 1;
+			UPROPERTY(Category = "Group", EditAnywhere, BlueprintReadOnly)
+			uint32 bGroup24 : 1;
+			UPROPERTY(Category = "Group", EditAnywhere, BlueprintReadOnly)
+			uint32 bGroup25 : 1;
+			UPROPERTY(Category = "Group", EditAnywhere, BlueprintReadOnly)
+			uint32 bGroup26 : 1;
+			UPROPERTY(Category = "Group", EditAnywhere, BlueprintReadOnly)
+			uint32 bGroup27 : 1;
+			UPROPERTY(Category = "Group", EditAnywhere, BlueprintReadOnly)
+			uint32 bGroup28 : 1;
+			UPROPERTY(Category = "Group", EditAnywhere, BlueprintReadOnly)
+			uint32 bGroup29 : 1;
+			UPROPERTY(Category = "Group", EditAnywhere, BlueprintReadOnly)
+			uint32 bGroup30 : 1;
+			UPROPERTY(Category = "Group", EditAnywhere, BlueprintReadOnly)
+			uint32 bGroup31 : 1;
+#if CPP
+		};
+		int32 Packed;
+	};
+#endif
+
+	FORCEINLINE bool HasGroup(uint8 GroupID) const
+	{
+		return (Packed & (1 << GroupID)) != 0;
+	}
+
+	FORCEINLINE void SetGroup(uint8 GroupID)
+	{
+		Packed |= (1 << GroupID);
+	}
+
+	FORCEINLINE void ClearGroup(uint8 GroupID)
+	{
+		Packed &= ~(1 << GroupID);
+	}
+
+	FORCEINLINE void ClearAll()
+	{
+		Packed = 0;
+	}
+
+	FORCEINLINE void SetFlagsDirectly(uint32 NewFlagset)
+	{
+		Packed = NewFlagset;
+	}
+};
diff --git a/Engine/Source/Runtime/Engine/Classes/AI/Navigation/NavigationTypes.h b/Engine/Source/Runtime/Engine/Classes/AI/Navigation/NavigationTypes.h
--- a/Engine/Source/Runtime/Engine/Classes/AI/Navigation/NavigationTypes.h
+++ b/Engine/Source/Runtime/Engine/Classes/AI/Navigation/NavigationTypes.h
@@ -4,6 +4,7 @@
 
 #include "NavLinkDefinition.h"
 #include "NavigationModifier.h"
+#include "NavigationAvoidanceTypes.h"
 #include "NavigationTypes.generated.h"
 
 #define INVALID_NAVNODEREF (0)
@@ -212,6 +213,7 @@ struct ENGINE_API FNavigationPath : public TSharedFromThis<FNavigationPath, ESPM
 	FORCEINLINE class ANavigationData *GetOwner() const { return Owner.Get(); }
 	FORCEINLINE bool IsDirect() const { return Owner.Get() == NULL; }
 	FORCEINLINE FVector GetDestinationLocation() const { return IsValid() ? PathPoints.Last().Location : INVALID_NAVEXTENT; }
+	FORCEINLINE void GetObserver(FPathObserverDelegate& Observer) const { Observer = ObserverDelegate; }
 
 	FORCEINLINE void MarkReady() { bIsReady = true; }
 	FORCEINLINE void SetOwner(const class ANavigationData* const NewOwner) { Owner = NewOwner; }
diff --git a/Engine/Source/Runtime/Engine/Classes/GameFramework/CharacterMovementComponent.h b/Engine/Source/Runtime/Engine/Classes/GameFramework/CharacterMovementComponent.h
--- a/Engine/Source/Runtime/Engine/Classes/GameFramework/CharacterMovementComponent.h
+++ b/Engine/Source/Runtime/Engine/Classes/GameFramework/CharacterMovementComponent.h
@@ -7,8 +7,8 @@
 
 #pragma once
 
+#include "AI/Navigation/NavigationAvoidanceTypes.h"
 #include "Interfaces/NetworkPredictionInterface.h"
-
 #include "CharacterMovementComponent.generated.h"
 
 /** Movement modes for Characters.  */
@@ -24,113 +24,6 @@ enum EMovementMode
 	MOVE_MAX		UMETA(Hidden),
 };
 
-USTRUCT()
-struct FNavAvoidanceMask
-{
-	GENERATED_USTRUCT_BODY()
-
-#if CPP
-	union
-	{
-		struct
-		{
-#endif
-		UPROPERTY(Category="Group", EditAnywhere, BlueprintReadOnly)
-		uint32 bGroup0 : 1;
-		UPROPERTY(Category="Group", EditAnywhere, BlueprintReadOnly)
-		uint32 bGroup1 : 1;
-		UPROPERTY(Category="Group", EditAnywhere, BlueprintReadOnly)
-		uint32 bGroup2 : 1;
-		UPROPERTY(Category="Group", EditAnywhere, BlueprintReadOnly)
-		uint32 bGroup3 : 1;
-		UPROPERTY(Category="Group", EditAnywhere, BlueprintReadOnly)
-		uint32 bGroup4 : 1;
-		UPROPERTY(Category="Group", EditAnywhere, BlueprintReadOnly)
-		uint32 bGroup5 : 1;
-		UPROPERTY(Category="Group", EditAnywhere, BlueprintReadOnly)
-		uint32 bGroup6 : 1;
-		UPROPERTY(Category="Group", EditAnywhere, BlueprintReadOnly)
-		uint32 bGroup7 : 1;
-		UPROPERTY(Category="Group", EditAnywhere, BlueprintReadOnly)
-		uint32 bGroup8 : 1;
-		UPROPERTY(Category="Group", EditAnywhere, BlueprintReadOnly)
-		uint32 bGroup9 : 1;
-		UPROPERTY(Category="Group", EditAnywhere, BlueprintReadOnly)
-		uint32 bGroup10 : 1;
-		UPROPERTY(Category="Group", EditAnywhere, BlueprintReadOnly)
-		uint32 bGroup11 : 1;
-		UPROPERTY(Category="Group", EditAnywhere, BlueprintReadOnly)
-		uint32 bGroup12 : 1;
-		UPROPERTY(Category="Group", EditAnywhere, BlueprintReadOnly)
-		uint32 bGroup13 : 1;
-		UPROPERTY(Category="Group", EditAnywhere, BlueprintReadOnly)
-		uint32 bGroup14 : 1;
-		UPROPERTY(Category="Group", EditAnywhere, BlueprintReadOnly)
-		uint32 bGroup15 : 1;
-		UPROPERTY(Category="Group", EditAnywhere, BlueprintReadOnly)
-		uint32 bGroup16 : 1;
-		UPROPERTY(Category="Group", EditAnywhere, BlueprintReadOnly)
-		uint32 bGroup17 : 1;
-		UPROPERTY(Category="Group", EditAnywhere, BlueprintReadOnly)
-		uint32 bGroup18 : 1;
-		UPROPERTY(Category="Group", EditAnywhere, BlueprintReadOnly)
-		uint32 bGroup19 : 1;
-		UPROPERTY(Category="Group", EditAnywhere, BlueprintReadOnly)
-		uint32 bGroup20 : 1;
-		UPROPERTY(Category="Group", EditAnywhere, BlueprintReadOnly)
-		uint32 bGroup21 : 1;
-		UPROPERTY(Category="Group", EditAnywhere, BlueprintReadOnly)
-		uint32 bGroup22 : 1;
-		UPROPERTY(Category="Group", EditAnywhere, BlueprintReadOnly)
-		uint32 bGroup23 : 1;
-		UPROPERTY(Category="Group", EditAnywhere, BlueprintReadOnly)
-		uint32 bGroup24 : 1;
-		UPROPERTY(Category="Group", EditAnywhere, BlueprintReadOnly)
-		uint32 bGroup25 : 1;
-		UPROPERTY(Category="Group", EditAnywhere, BlueprintReadOnly)
-		uint32 bGroup26 : 1;
-		UPROPERTY(Category="Group", EditAnywhere, BlueprintReadOnly)
-		uint32 bGroup27 : 1;
-		UPROPERTY(Category="Group", EditAnywhere, BlueprintReadOnly)
-		uint32 bGroup28 : 1;
-		UPROPERTY(Category="Group", EditAnywhere, BlueprintReadOnly)
-		uint32 bGroup29 : 1;
-		UPROPERTY(Category="Group", EditAnywhere, BlueprintReadOnly)
-		uint32 bGroup30 : 1;
-		UPROPERTY(Category="Group", EditAnywhere, BlueprintReadOnly)
-		uint32 bGroup31 : 1;
-#if CPP
-		};
-		int32 Packed;
-	};
-#endif
-
-	FORCEINLINE bool HasGroup(uint8 GroupID) const
-	{
-		return (Packed & (1 << GroupID)) != 0;
-	}
-
-	FORCEINLINE void SetGroup(uint8 GroupID)
-	{
-		Packed |= (1 << GroupID);
-	}
-
-	FORCEINLINE void ClearGroup(uint8 GroupID)
-	{
-		Packed &= ~(1 << GroupID);
-	}
-
-	FORCEINLINE void ClearAll()
-	{
-		Packed = 0;
-	}
-
-	FORCEINLINE void SetFlagsDirectly(uint32 NewFlagset)
-	{
-		Packed = NewFlagset;
-	}
-};
-
 // Data about the floor for walking movement.
 USTRUCT(BlueprintType)
 struct ENGINE_API FFindFloorResult
diff --git a/Engine/Source/Runtime/Engine/Classes/GameFramework/CheatManager.h b/Engine/Source/Runtime/Engine/Classes/GameFramework/CheatManager.h
--- a/Engine/Source/Runtime/Engine/Classes/GameFramework/CheatManager.h
+++ b/Engine/Source/Runtime/Engine/Classes/GameFramework/CheatManager.h
@@ -77,6 +77,9 @@ class ENGINE_API UCheatManager : public UObject
 	/** If we should trace complex collision in debug capsule sweeps. Set with DebugCapsuleSweepComplex() */
 	uint32 bDebugCapsuleTraceComplex:1;
 
+	/** Holds information if we used ToggleAILogging cheat to activate AI logging */
+	uint32 bToggleAILogging : 1;
+
 	/** How far debugf trace should go out from player viewpoint */
 	float DebugTraceDistance;
 
@@ -332,6 +335,12 @@ class ENGINE_API UCheatManager : public UObject
 	 */
 	virtual void InitCheatManager();
 
+	/**
+	* Called before destroying the object.  This is called immediately upon deciding to destroy the object, to allow the object to begin an
+	* asynchronous cleanup process.
+	*/
+	virtual void BeginDestroy() override;
+
 	/** Use the Outer Player Controller to get a World.  */
 	UWorld* GetWorld() const;
 protected:
diff --git a/Engine/Source/Runtime/Engine/Private/CheatManager.cpp b/Engine/Source/Runtime/Engine/Private/CheatManager.cpp
--- a/Engine/Source/Runtime/Engine/Private/CheatManager.cpp
+++ b/Engine/Source/Runtime/Engine/Private/CheatManager.cpp
@@ -553,7 +553,31 @@ void UCheatManager::DisableDebugCamera()
 	}
 }
 
-void UCheatManager::InitCheatManager() {}
+void UCheatManager::InitCheatManager() 
+{
+
+}
+
+void UCheatManager::BeginDestroy()
+{
+#if ENABLE_VISUAL_LOG
+	FVisualLog* VisLog = FVisualLog::Get();
+	if (VisLog && VisLog->IsRecording() && bToggleAILogging)
+	{
+		UWorld *World = GetWorld();
+		if (World)
+		{
+			// clear timer, we'll dump all remaining logs
+			World->GetTimerManager().ClearTimer(this, &UCheatManager::DumpAILogs);
+		}
+
+		// stop recording and dump all remaining logs
+		VisLog->SetIsRecording(false);
+		bToggleAILogging = false;
+	}
+#endif
+	Super::BeginDestroy();
+}
 
 bool UCheatManager::ServerToggleAILogging_Validate()
 {
@@ -595,14 +619,16 @@ void UCheatManager::ServerToggleAILogging_Implementation()
 
 		// stop recording and dump all remaining logs in a moment
 		VisLog->SetIsRecording(false, true);
+		bToggleAILogging = false;
 	}
 	else
 	{
-		VisLog->SetIsRecording(true, true);
 		if (World)
 		{
 			World->GetTimerManager().SetTimer(this, &UCheatManager::DumpAILogs, DumpAILogsInterval);
 		}
+		VisLog->SetIsRecording(true, true);
+		bToggleAILogging = true;
 	}
 
 	GetOuterAPlayerController()->ClientMessage(FString::Printf(TEXT("OK! VisLog recording is now %s"), VisLog->IsRecording() ? TEXT("Enabled") : TEXT("Disabled")));
diff --git a/Engine/Source/Runtime/Engine/Private/Debug/ReporterGraph.cpp b/Engine/Source/Runtime/Engine/Private/Debug/ReporterGraph.cpp
--- a/Engine/Source/Runtime/Engine/Private/Debug/ReporterGraph.cpp
+++ b/Engine/Source/Runtime/Engine/Private/Debug/ReporterGraph.cpp
@@ -150,33 +150,29 @@ void UReporterGraph::DrawLegend(UCanvas* Canvas)
 	FVector2D CurrentTextPos = GraphScreenSize.Max;
 	CurrentTextPos.X = GraphScreenSize.Min.X;
 	UFont* Font = GEngine->GetSmallFont();
-	if(LegendWidth < KINDA_SMALL_NUMBER)
+
+	int32 DummyY, CurrentX;
+	StringSize(Font, CurrentX, DummyY, TEXT("99.99"));
+	const float InsideLegendWidth = -CurrentX;
+
+	for(int32 i = 0; i < CurrentData.Num(); i++)
 	{
+		FVector2D ScreenPos = ToScreenSpace(CurrentTextPos, Canvas);
 		if (LegendPosition == ELegendPosition::Outside)
 		{
 			int32 DummyY, CurrentX;
-			for (int32 i = 0; i < CurrentData.Num(); i++)
-			{
-				StringSize(Font, CurrentX, DummyY, *CurrentData[i].LineName);
-
-				if (CurrentX > LegendWidth)
-				{
-					LegendWidth = CurrentX;
-				}
-			}
+			StringSize(Font, CurrentX, DummyY, *CurrentData[i].LineName);
+			LegendWidth = CurrentX + 10;
 		}
 		else
 		{
-			int32 DummyY, CurrentX;
-			StringSize(Font, CurrentX, DummyY, TEXT("M"));
-			LegendWidth = -CurrentX;
+			LegendWidth = InsideLegendWidth;
 		}
-	}
-	
-	for(int32 i = 0; i < CurrentData.Num(); i++)
-	{
-		FVector2D ScreenPos = ToScreenSpace(CurrentTextPos, Canvas);
-		Canvas->Canvas->DrawShadowedString(ScreenPos.X  - LegendWidth, ScreenPos.Y, *CurrentData[i].LineName, Font, CurrentData[i].Color);
+
+		FCanvasTextItem TextItem(FVector2D::ZeroVector, FText::FromString(*CurrentData[i].LineName), Font, CurrentData[i].Color);
+		TextItem.EnableShadow(FColor::Black, FVector2D(1, 1));
+		TextItem.SetColor(CurrentData[i].Color);
+		Canvas->DrawItem(TextItem, ScreenPos.X - LegendWidth, ScreenPos.Y);
 
 		CurrentTextPos.Y -= (GraphScreenSize.Max.Y - GraphScreenSize.Min.Y) / CurrentData.Num();
 	}
@@ -195,9 +191,19 @@ void UReporterGraph::DrawAxes(UCanvas* Canvas)
 	YMax.Y = GraphScreenSize.Max.Y;
 
 	// Draw the X axis
+	const float SizeX = (XMax.X - Min.X) * Canvas->SizeX;
+	int32 StringSizeX, StringSizeY;
+	UFont* Font = GEngine->GetSmallFont();
+
+	StringSize(Font, StringSizeX, StringSizeY, *FString::Printf(TEXT("%.2f"), GraphMinMaxData.Max.X) );
+	NumXNotches = FMath::CeilToInt(SizeX * 0.7 / StringSizeX);
+
 	DrawAxis(Canvas, Min, XMax, NumXNotches, false);
 
 	// Draw the Y axis
+	StringSize(Font, StringSizeX, StringSizeY, *FString::Printf(TEXT("%.2f"), GraphMinMaxData.Max.Y));
+	float SizeY = (YMax.Y - Min.Y) * Canvas->SizeY;
+	NumYNotches = FMath::CeilToInt(SizeY * 0.7 / StringSizeY);
 	DrawAxis(Canvas, Min, YMax, NumYNotches, true);
 }
 
@@ -275,7 +281,7 @@ void UReporterGraph::DrawAxis(UCanvas* Canvas, FVector2D Start, FVector2D End, f
 			NotchLocation.X += NotchDelta;
 		}
 
-		FString NotchValue = FString::Printf(TEXT("%1.0f"), (bIsVerticalAxis ? GraphMinMaxData.Min.Y + (NotchDataDelta.Y * (i + 1)) : GraphMinMaxData.Min.X + (NotchDataDelta.X * (i + 1))));
+		FString NotchValue = FString::Printf(TEXT("%1.2f"), (bIsVerticalAxis ? GraphMinMaxData.Min.Y + (NotchDataDelta.Y * (i + 1)) : GraphMinMaxData.Min.X + (NotchDataDelta.X * (i + 1))));
 		
 		int32 StringSizeX, StringSizeY;
 		StringSize(Font, StringSizeX, StringSizeY, *NotchValue );
diff --git a/Engine/Source/Runtime/Engine/Private/UserInterface/CanvasItem.cpp b/Engine/Source/Runtime/Engine/Private/UserInterface/CanvasItem.cpp
--- a/Engine/Source/Runtime/Engine/Private/UserInterface/CanvasItem.cpp
+++ b/Engine/Source/Runtime/Engine/Private/UserInterface/CanvasItem.cpp
@@ -776,6 +776,7 @@ void FCanvasTextItem::Draw( class FCanvas* InCanvas )
 void FCanvasTextItem::DrawStringInternal( class FCanvas* InCanvas, const FVector2D& DrawPos, const FLinearColor& InColour )
 {
 	DrawnSize = FVector2D(EForceInit::ForceInitToZero);
+	FVector2D CurrentPos = FVector2D(EForceInit::ForceInitToZero);
 	FHitProxyId HitProxyId = InCanvas->GetHitProxyId();
 	FTexture* LastTexture = NULL;
 	UTexture2D* Tex = NULL;
@@ -794,11 +795,20 @@ void FCanvasTextItem::DrawStringInternal( class FCanvas* InCanvas, const FVector
 
 		FFontCharacter& Char = Font->Characters[Ch];
 
+		if (DrawnSize.Y == 0)
+		{
+			// We have a valid character so initialize vertical DrawnSize
+			DrawnSize.Y = Font->GetMaxCharHeight() * YScale;
+		}
+
 		if (FChar::IsLinebreak(Chars[i]))
 		{
 			// Set current character offset to the beginning of next line.
-			DrawnSize.X = 0.0f;
-			DrawnSize.Y += Char.VSize * YScale;
+			CurrentPos.X = 0.0f;
+			CurrentPos.Y += Font->GetMaxCharHeight() * YScale;
+
+			// Increase the vertical DrawnSize
+			DrawnSize.Y += Font->GetMaxCharHeight() * YScale;
 
 			// Don't draw newline character
 			continue;
@@ -824,8 +834,8 @@ void FCanvasTextItem::DrawStringInternal( class FCanvas* InCanvas, const FVector
 			}
 			LastTexture = Tex->Resource;
 
-			const float X		= DrawnSize.X + DrawPos.X;
-			const float Y		= DrawnSize.Y + DrawPos.Y + Char.VerticalOffset * YScale;
+			const float X		= CurrentPos.X + DrawPos.X;
+			const float Y		= CurrentPos.Y + DrawPos.Y + Char.VerticalOffset * YScale;
 			float SizeX			= Char.USize * XScale;
 			const float SizeY	= Char.VSize * YScale;
 			const float U		= Char.StartU * InvTextureSize.X;
@@ -870,7 +880,13 @@ void FCanvasTextItem::DrawStringInternal( class FCanvas* InCanvas, const FVector
 			}
 
 			// Update the current rendering position
-			DrawnSize.X += SizeX;
+			CurrentPos.X += SizeX;
+
+			// Increase the Horizontal DrawnSize
+			if (CurrentPos.X > DrawnSize.X)
+			{
+				DrawnSize.X = CurrentPos.X;
+			}
 		}
 	}
 }
diff --git a/Engine/Source/Runtime/GameplayTags/Classes/GameplayTagsManager.h b/Engine/Source/Runtime/GameplayTags/Classes/GameplayTagsManager.h
--- a/Engine/Source/Runtime/GameplayTags/Classes/GameplayTagsManager.h
+++ b/Engine/Source/Runtime/GameplayTags/Classes/GameplayTagsManager.h
@@ -232,6 +232,11 @@ class GAMEPLAYTAGS_API UGameplayTagsManager : public UObject
 
 private:
 
+	friend class FGameplayTagTest;
+
+	/** Helper function to populate the tag tree from each table */
+	void PopulateTreeFromDataTable(class UDataTable* Table);
+
 	/**
 	 * Helper function for RequestGameplayTagParents to add all parents to the container
 	 *
diff --git a/Engine/Source/Runtime/GameplayTags/Private/GameplayTagTests.cpp b/Engine/Source/Runtime/GameplayTags/Private/GameplayTagTests.cpp
new file mode 100644
--- /dev/null
+++ b/Engine/Source/Runtime/GameplayTags/Private/GameplayTagTests.cpp
@@ -0,0 +1,43 @@
+// Copyright 1998-2014 Epic Games, Inc. All Rights Reserved.
+
+#include "GameplayTagsModulePrivatePCH.h"
+
+IMPLEMENT_SIMPLE_AUTOMATION_TEST(FGameplayTagTest, "GameplayTags.GameplayTag", EAutomationTestFlags::ATF_Editor)
+
+UDataTable* CreateGameplayDataTable()
+{
+	FString CSV(TEXT(",Tag,CategoryText,\r\n0,GameplayTagTest.Test1\r\n1,GameplayTagTest.Test2"));
+
+	UDataTable * DataTable = Cast<UDataTable>(StaticConstructObject(UDataTable::StaticClass(), GetTransientPackage(), FName(TEXT("TempDataTable"))));
+	DataTable->RowStruct = FGameplayTagTableRow::StaticStruct();
+	DataTable->CreateTableFromCSVString(CSV);
+
+	FGameplayTagTableRow * Row = (FGameplayTagTableRow*)DataTable->RowMap["0"];
+	if (Row)
+	{
+		check(Row->Tag == TEXT("GameplayTagTest.Test1"));
+	}
+	return DataTable;
+}
+
+bool GameplayTagTest_SimpleTest()
+{
+	FName TagName = FName(TEXT("GameplayTagTest.Test2"));
+	FGameplayTag Tag = IGameplayTagsModule::Get().GetGameplayTagsManager().RequestGameplayTag(TagName);
+	return Tag.GetTagName() == TagName;
+}
+
+bool FGameplayTagTest::RunTest(const FString& Parameters)
+{
+	// Create Test Data 
+	UDataTable* DataTable = CreateGameplayDataTable();
+
+	IGameplayTagsModule::Get().GetGameplayTagsManager().PopulateTreeFromDataTable(DataTable);
+
+	// Run Tests
+	bool bSuccess = true;
+	bSuccess &= GameplayTagTest_SimpleTest();
+	// Add more tests here... 
+
+	return bSuccess;
+}
diff --git a/Engine/Source/Runtime/GameplayTags/Private/GameplayTagsManager.cpp b/Engine/Source/Runtime/GameplayTags/Private/GameplayTagsManager.cpp
--- a/Engine/Source/Runtime/GameplayTags/Private/GameplayTagsManager.cpp
+++ b/Engine/Source/Runtime/GameplayTags/Private/GameplayTagsManager.cpp
@@ -112,46 +112,50 @@ void UGameplayTagsManager::ConstructGameplayTagTree()
 	if (!GameplayRootTag.IsValid() && GameplayTagTables.Num() > 0)
 	{
 		GameplayRootTag = MakeShareable(new FGameplayTagNode());
-		TArray< TSharedPtr<FGameplayTagNode> >& GameplayRootTags = GameplayRootTag->GetChildTagNodes();
-		TSet<FName> RootTags;
-
+		
 		for (auto It(GameplayTagTables.CreateIterator()); It; It++)
 		{
-			static const FString ContextString(TEXT("UNKNOWN"));
-			const int32 NumRows = (*It)->RowMap.Num();
-			for (int32 RowIdx = 0; RowIdx < NumRows; ++RowIdx)
+			PopulateTreeFromDataTable(*It);
+		}
+	}
+}
+
+void UGameplayTagsManager::PopulateTreeFromDataTable(class UDataTable* InTable)
+{
+	TArray< TSharedPtr<FGameplayTagNode> >& GameplayRootTags = GameplayRootTag->GetChildTagNodes();
+	static const FString ContextString(TEXT("UNKNOWN"));
+	const int32 NumRows = InTable->RowMap.Num();
+	for (int32 RowIdx = 0; RowIdx < NumRows; ++RowIdx)
+	{
+		FGameplayTagTableRow* TagRow = InTable->FindRow<FGameplayTagTableRow>(*FString::Printf(TEXT("%d"), RowIdx), ContextString);
+		if (TagRow)
+		{
+			// Split the tag text on the "." delimiter to establish tag depth and then insert each tag into the
+			// gameplay tag tree
+			TArray<FString> SubTags;
+			TagRow->Tag.ParseIntoArray(&SubTags, TEXT("."), true);
+
+			if (SubTags.Num() > 0)
 			{
-				FGameplayTagTableRow* TagRow = (*It)->FindRow<FGameplayTagTableRow>(*FString::Printf(TEXT("%d"), RowIdx), ContextString);
-				if (TagRow)
-				{
-					// Split the tag text on the "." delimiter to establish tag depth and then insert each tag into the
-					// gameplay tag tree
-					TArray<FString> SubTags;
-					TagRow->Tag.ParseIntoArray(&SubTags, TEXT("."), true);
+				int32 InsertionIdx = InsertTagIntoNodeArray(*SubTags[0], NULL, GameplayRootTags, SubTags.Num() == 1 ? TagRow->CategoryText : FText());
+				TSharedPtr<FGameplayTagNode> CurNode = GameplayRootTags[InsertionIdx];
 
-					if (SubTags.Num() > 0)
+				for (int32 SubTagIdx = 1; SubTagIdx < SubTags.Num(); ++SubTagIdx)
+				{
+					TArray< TSharedPtr<FGameplayTagNode> >& ChildTags = CurNode.Get()->GetChildTagNodes();
+					FText Description;
+					if (SubTagIdx == SubTags.Num() - 1)
 					{
-						int32 InsertionIdx = InsertTagIntoNodeArray(*SubTags[0], NULL, GameplayRootTags, SubTags.Num() == 1 ? TagRow->CategoryText : FText());
-						TSharedPtr<FGameplayTagNode> CurNode = GameplayRootTags[InsertionIdx];
-
-						for (int32 SubTagIdx = 1; SubTagIdx < SubTags.Num(); ++SubTagIdx)
-						{
-							TArray< TSharedPtr<FGameplayTagNode> >& ChildTags = CurNode.Get()->GetChildTagNodes();
-							FText Description;
-							if (SubTagIdx == SubTags.Num() - 1)
-							{
-								Description = TagRow->CategoryText;
-							}
-							InsertionIdx = InsertTagIntoNodeArray(*SubTags[SubTagIdx], CurNode, ChildTags, Description);
-
-							CurNode = ChildTags[InsertionIdx];
-						}
+						Description = TagRow->CategoryText;
 					}
+					InsertionIdx = InsertTagIntoNodeArray(*SubTags[SubTagIdx], CurNode, ChildTags, Description);
+
+					CurNode = ChildTags[InsertionIdx];
 				}
 			}
 		}
-		GameplayRootTags.Sort(FCompareFGameplayTagNodeByTag());
 	}
+	GameplayRootTags.Sort(FCompareFGameplayTagNodeByTag());
 }
 
 UGameplayTagsManager::~UGameplayTagsManager()
diff --git a/Engine/Source/Runtime/Navmesh/Private/DetourCrowd/DetourCrowd.cpp b/Engine/Source/Runtime/Navmesh/Private/DetourCrowd/DetourCrowd.cpp
--- a/Engine/Source/Runtime/Navmesh/Private/DetourCrowd/DetourCrowd.cpp
+++ b/Engine/Source/Runtime/Navmesh/Private/DetourCrowd/DetourCrowd.cpp
@@ -217,6 +217,11 @@ static int getNeighbours(const float* pos, const float height, const float range
 		const float distSqr = dtVlenSqr(diff);
 		if (distSqr > dtSqr(range))
 			continue;
+
+		// [UE4] add only when avoidance group allows it
+		const bool bDontAvoid = (skip->params.groupsToIgnore & ag->params.avoidanceGroup) || !(skip->params.groupsToAvoid & ag->params.avoidanceGroup);
+		if (bDontAvoid)
+			continue;
 		
 		n = addNeighbour(ids[i], distSqr, result, n, maxResult);
 	}
diff --git a/Engine/Source/Runtime/Navmesh/Public/DetourCrowd/DetourCrowd.h b/Engine/Source/Runtime/Navmesh/Public/DetourCrowd/DetourCrowd.h
--- a/Engine/Source/Runtime/Navmesh/Public/DetourCrowd/DetourCrowd.h
+++ b/Engine/Source/Runtime/Navmesh/Public/DetourCrowd/DetourCrowd.h
@@ -93,8 +93,15 @@ struct dtCrowdAgentParams
 	/// How aggresive the agent manager should be at avoiding collisions with this agent. [Limit: >= 0]
 	float separationWeight;
 
+	/// [UE4] Groups flags attached to the agent
+	unsigned int avoidanceGroup;
+	/// [UE4] Avoid agents when they group is matching mask
+	unsigned int groupsToAvoid;
+	/// [UE4] Don't avoid agents when they group is matching mask
+	unsigned int groupsToIgnore;
+
 	/// Flags that impact steering behavior. (See: #UpdateFlags)
-	unsigned char updateFlags;
+	unsigned short updateFlags;
 
 	/// The index of the avoidance configuration to use for the agent. 
 	/// [Limits: 0 <= value <= #DT_CROWD_MAX_OBSTAVOIDANCE_PARAMS]
diff --git a/Engine/Source/Runtime/Windows/XAudio2/Private/XAudio2Source.cpp b/Engine/Source/Runtime/Windows/XAudio2/Private/XAudio2Source.cpp
--- a/Engine/Source/Runtime/Windows/XAudio2/Private/XAudio2Source.cpp
+++ b/Engine/Source/Runtime/Windows/XAudio2/Private/XAudio2Source.cpp
@@ -1226,6 +1226,154 @@ void FSpatializationHelper::Init()
 	DSPSettings.pMatrixCoefficients = MatrixCoefficients;
 }
 
+void FSpatializationHelper::DumpSpatializationState() const
+{
+	struct FLocal
+	{
+		static void DumpCone(const FString& Indent, const FString& ConeName, const X3DAUDIO_CONE* pCone)
+		{
+			if (pCone)
+			{
+				UE_LOG(LogXAudio2, Log, TEXT("%s%s"), *Indent, *ConeName);
+				UE_LOG(LogXAudio2, Log, TEXT("%s  InnerAngle: %f"), *Indent, pCone->InnerAngle);
+				UE_LOG(LogXAudio2, Log, TEXT("%s  OuterAngle: %f"), *Indent, pCone->OuterAngle);
+				UE_LOG(LogXAudio2, Log, TEXT("%s  InnerVolume: %f"), *Indent, pCone->InnerVolume);
+				UE_LOG(LogXAudio2, Log, TEXT("%s  OuterVolume: %f"), *Indent, pCone->OuterVolume);
+				UE_LOG(LogXAudio2, Log, TEXT("%s  InnerLPF: %f"), *Indent, pCone->InnerLPF);
+				UE_LOG(LogXAudio2, Log, TEXT("%s  OuterLPF: %f"), *Indent, pCone->OuterLPF);
+				UE_LOG(LogXAudio2, Log, TEXT("%s  InnerReverb: %f"), *Indent, pCone->InnerReverb);
+				UE_LOG(LogXAudio2, Log, TEXT("%s  OuterReverb: %f"), *Indent, pCone->OuterReverb);
+			}
+			else
+			{
+				UE_LOG(LogXAudio2, Log, TEXT("%s%s: NULL"), *Indent, *ConeName);
+			}
+		}
+
+		static void DumpDistanceCurvePoint(const FString& Indent, const FString& PointName, uint32 Index, const X3DAUDIO_DISTANCE_CURVE_POINT& Point)
+		{
+			UE_LOG(LogXAudio2, Log, TEXT("%s%s[%u]: {%f,%f}"), *Indent, *PointName, Index, Point.Distance, Point.DSPSetting);
+		}
+
+		static void DumpDistanceCurve(const FString& Indent, const FString& CurveName, const X3DAUDIO_DISTANCE_CURVE* pCurve)
+		{
+			const uint32 MaxPointsToDump = 20;
+			if (pCurve)
+			{
+				UE_LOG(LogXAudio2, Log, TEXT("%s%s: %u points"), *Indent, *CurveName, pCurve->PointCount);
+				for (uint32 PointIdx = 0; PointIdx < pCurve->PointCount && PointIdx < MaxPointsToDump; ++PointIdx)
+				{
+					const X3DAUDIO_DISTANCE_CURVE_POINT& CurPoint = pCurve->pPoints[PointIdx];
+					DumpDistanceCurvePoint(Indent + TEXT("  "), TEXT("pPoints"), PointIdx, CurPoint);
+				}
+			}
+			else
+			{
+				UE_LOG(LogXAudio2, Log, TEXT("%s%s: NULL"), *Indent, *CurveName);
+			}
+		}
+	};
+
+	UE_LOG(LogXAudio2, Log, TEXT("Dumping all XAudio2 Spatialization"));
+	UE_LOG(LogXAudio2, Log, TEXT("==================================="));
+	
+	// X3DInstance
+	UE_LOG(LogXAudio2, Log, TEXT("  X3DInstance: %#010x"), X3DInstance);
+
+	// DSPSettings
+	UE_LOG(LogXAudio2, Log, TEXT("  DSPSettings"));
+	
+	if ( DSPSettings.pMatrixCoefficients )
+	{
+		UE_LOG(LogXAudio2, Log, TEXT("    pMatrixCoefficients: {%f,%f,%f,%f,%f,%f}"), DSPSettings.pMatrixCoefficients[0], DSPSettings.pMatrixCoefficients[1], DSPSettings.pMatrixCoefficients[2], DSPSettings.pMatrixCoefficients[3], DSPSettings.pMatrixCoefficients[4], DSPSettings.pMatrixCoefficients[5]);
+	}
+	else
+	{
+		UE_LOG(LogXAudio2, Log, TEXT("    pMatrixCoefficients: NULL"));
+	}
+
+	if (DSPSettings.pDelayTimes)
+	{
+		UE_LOG(LogXAudio2, Log, TEXT("    pDelayTimes: {%f,%f,%f,%f,%f,%f}"), DSPSettings.pDelayTimes[0], DSPSettings.pDelayTimes[1], DSPSettings.pDelayTimes[2], DSPSettings.pDelayTimes[3], DSPSettings.pDelayTimes[4], DSPSettings.pDelayTimes[5]);
+	}
+	else
+	{
+		UE_LOG(LogXAudio2, Log, TEXT("    pDelayTimes: NULL"));
+	}
+
+	UE_LOG(LogXAudio2, Log, TEXT("    SrcChannelCount: %u"), DSPSettings.SrcChannelCount);
+	UE_LOG(LogXAudio2, Log, TEXT("    DstChannelCount: %u"), DSPSettings.DstChannelCount);
+	UE_LOG(LogXAudio2, Log, TEXT("    LPFDirectCoefficient: %f"), DSPSettings.LPFDirectCoefficient);
+	UE_LOG(LogXAudio2, Log, TEXT("    LPFReverbCoefficient: %f"), DSPSettings.LPFReverbCoefficient);
+	UE_LOG(LogXAudio2, Log, TEXT("    ReverbLevel: %f"), DSPSettings.ReverbLevel);
+	UE_LOG(LogXAudio2, Log, TEXT("    DopplerFactor: %f"), DSPSettings.DopplerFactor);
+	UE_LOG(LogXAudio2, Log, TEXT("    EmitterToListenerAngle: %f"), DSPSettings.EmitterToListenerAngle);
+	UE_LOG(LogXAudio2, Log, TEXT("    EmitterToListenerDistance: %f"), DSPSettings.EmitterToListenerDistance);
+	UE_LOG(LogXAudio2, Log, TEXT("    EmitterVelocityComponent: %f"), DSPSettings.EmitterVelocityComponent);
+	UE_LOG(LogXAudio2, Log, TEXT("    ListenerVelocityComponent: %f"), DSPSettings.ListenerVelocityComponent);
+
+	// Listener
+	UE_LOG(LogXAudio2, Log, TEXT("  Listener"));
+	UE_LOG(LogXAudio2, Log, TEXT("    OrientFront: {%f,%f,%f}"), Listener.OrientFront.x, Listener.OrientFront.y, Listener.OrientFront.z);
+	UE_LOG(LogXAudio2, Log, TEXT("    OrientTop: {%f,%f,%f}"), Listener.OrientTop.x, Listener.OrientTop.y, Listener.OrientTop.z);
+	UE_LOG(LogXAudio2, Log, TEXT("    Position: {%f,%f,%f}"), Listener.Position.x, Listener.Position.y, Listener.Position.z);
+	UE_LOG(LogXAudio2, Log, TEXT("    Velocity: {%f,%f,%f}"), Listener.Velocity.x, Listener.Velocity.y, Listener.Velocity.z);
+	FLocal::DumpCone(TEXT("    "), TEXT("pCone"), Listener.pCone);
+
+	// Emitter
+	UE_LOG(LogXAudio2, Log, TEXT("  Emitter"));
+	FLocal::DumpCone(TEXT("    "), TEXT("pCone"), Emitter.pCone);
+	UE_LOG(LogXAudio2, Log, TEXT("    OrientFront: {%f,%f,%f}"), Emitter.OrientFront.x, Emitter.OrientFront.y, Emitter.OrientFront.z);
+	UE_LOG(LogXAudio2, Log, TEXT("    OrientTop: {%f,%f,%f}"), Emitter.OrientTop.x, Emitter.OrientTop.y, Emitter.OrientTop.z);
+	UE_LOG(LogXAudio2, Log, TEXT("    Position: {%f,%f,%f}"), Emitter.Position.x, Emitter.Position.y, Emitter.Position.z);
+	UE_LOG(LogXAudio2, Log, TEXT("    Velocity: {%f,%f,%f}"), Emitter.Velocity.x, Emitter.Velocity.y, Emitter.Velocity.z);
+	UE_LOG(LogXAudio2, Log, TEXT("    InnerRadius: %f"), Emitter.InnerRadius);
+	UE_LOG(LogXAudio2, Log, TEXT("    InnerRadiusAngle: %f"), Emitter.InnerRadiusAngle);
+	UE_LOG(LogXAudio2, Log, TEXT("    ChannelCount: %u"), Emitter.ChannelCount);
+	UE_LOG(LogXAudio2, Log, TEXT("    ChannelRadius: %f"), Emitter.ChannelRadius);
+
+	if ( Emitter.pChannelAzimuths )
+	{
+		UE_LOG(LogXAudio2, Log, TEXT("    pChannelAzimuths: %f"), *Emitter.pChannelAzimuths);
+	}
+	else
+	{
+		UE_LOG(LogXAudio2, Log, TEXT("    pChannelAzimuths: NULL"));
+	}
+
+	FLocal::DumpDistanceCurve(TEXT("    "), TEXT("pVolumeCurve"), Emitter.pVolumeCurve);
+	FLocal::DumpDistanceCurve(TEXT("    "), TEXT("pLFECurve"), Emitter.pLFECurve);
+	FLocal::DumpDistanceCurve(TEXT("    "), TEXT("pLPFDirectCurve"), Emitter.pLPFDirectCurve);
+	FLocal::DumpDistanceCurve(TEXT("    "), TEXT("pLPFReverbCurve"), Emitter.pLPFReverbCurve);
+	FLocal::DumpDistanceCurve(TEXT("    "), TEXT("pReverbCurve"), Emitter.pReverbCurve);
+
+	UE_LOG(LogXAudio2, Log, TEXT("    CurveDistanceScaler: %f"), Emitter.CurveDistanceScaler);
+	UE_LOG(LogXAudio2, Log, TEXT("    DopplerScaler: %f"), Emitter.DopplerScaler);
+
+	// Cone
+	FLocal::DumpCone(TEXT("  "), TEXT("Cone"), &Cone);
+
+	// VolumeCurvePoint
+	FLocal::DumpDistanceCurvePoint(TEXT("  "), TEXT("VolumeCurvePoint"), 0, VolumeCurvePoint[0]);
+	FLocal::DumpDistanceCurvePoint(TEXT("  "), TEXT("VolumeCurvePoint"), 1, VolumeCurvePoint[1]);
+	
+	// VolumeCurve
+	FLocal::DumpDistanceCurve(TEXT("  "), TEXT("VolumeCurve"), &VolumeCurve);
+
+	// ReverbVolumeCurvePoint
+	FLocal::DumpDistanceCurvePoint(TEXT("  "), TEXT("ReverbVolumeCurvePoint"), 0, ReverbVolumeCurvePoint[0]);
+	FLocal::DumpDistanceCurvePoint(TEXT("  "), TEXT("ReverbVolumeCurvePoint"), 1, ReverbVolumeCurvePoint[1]);
+
+	// ReverbVolumeCurve
+	FLocal::DumpDistanceCurve(TEXT("  "), TEXT("ReverbVolumeCurve"), &ReverbVolumeCurve);
+	
+	// EmitterAzimuths
+	UE_LOG(LogXAudio2, Log, TEXT("  EmitterAzimuths: %f"), EmitterAzimuths);
+	
+	// MatrixCoefficients
+	UE_LOG(LogXAudio2, Log, TEXT("  MatrixCoefficients: {%f,%f,%f,%f,%f,%f}"), MatrixCoefficients[0], MatrixCoefficients[1], MatrixCoefficients[2], MatrixCoefficients[3], MatrixCoefficients[4], MatrixCoefficients[5]);
+}
+
 /**
  * Calculates the spatialized volumes for each channel.
  *
@@ -1262,8 +1410,9 @@ void FSpatializationHelper::CalculateDolbySurroundRate( const FVector& OrientFro
 		if (!FMath::IsFinite(OutVolumes[SpeakerIndex]))
 		{
 			const FString NaNorINF = FMath::IsNaN(OutVolumes[SpeakerIndex]) ? TEXT("NaN") : TEXT("INF");
-			UE_LOG(LogXAudio2, Warning, TEXT("CalculateDolbySurroundRate generated a %s in channel %d. Emitter Values - Position.x:%f Position.y:%f Position.z:%f InnerRadius:%f OmniRadius:%f MatrixCoefficient:%f"),
-				*NaNorINF, SpeakerIndex, Emitter.Position.x, Emitter.Position.y, Emitter.Position.z, Emitter.InnerRadius, OmniRadius, DSPSettings.pMatrixCoefficients[SpeakerIndex]);
+			UE_LOG(LogXAudio2, Warning, TEXT("CalculateDolbySurroundRate generated a %s in channel %d. OmniRadius:%f MatrixCoefficient:%f"),
+				*NaNorINF, SpeakerIndex, OmniRadius, DSPSettings.pMatrixCoefficients[SpeakerIndex]);
+			DumpSpatializationState();
 		}
 #endif
 	}
diff --git a/Engine/Source/Runtime/Windows/XAudio2/Private/XAudio2Support.h b/Engine/Source/Runtime/Windows/XAudio2/Private/XAudio2Support.h
--- a/Engine/Source/Runtime/Windows/XAudio2/Private/XAudio2Support.h
+++ b/Engine/Source/Runtime/Windows/XAudio2/Private/XAudio2Support.h
@@ -448,6 +448,11 @@ public:
 
 	void Init();
 
+	/**
+	 * Logs out the entire state of the SpatializationHelper
+	 */
+	void DumpSpatializationState() const;
+
 	/**
 	 * Calculates the spatialized volumes for each channel.
 	 *
