fixed disappearing EQS test nodes after asset save
ttp# 346535

diff --git a/Engine/Source/Editor/EnvironmentQueryEditor/Private/EnvironmentQueryGraph.cpp b/Engine/Source/Editor/EnvironmentQueryEditor/Private/EnvironmentQueryGraph.cpp
--- a/Engine/Source/Editor/EnvironmentQueryEditor/Private/EnvironmentQueryGraph.cpp
+++ b/Engine/Source/Editor/EnvironmentQueryEditor/Private/EnvironmentQueryGraph.cpp
@@ -222,9 +222,7 @@ void UEnvironmentQueryGraph::RemoveOrphanedNodes()
 
 			for (int32 SubIdx = 0; SubIdx < OptionNode->Tests.Num(); SubIdx++)
 			{
-				if (OptionNode->Tests[SubIdx]  &&
-					OptionNode->Tests[SubIdx]->bTestEnabled &&
-					OptionNode->Tests[SubIdx]->NodeInstance)
+				if (OptionNode->Tests[SubIdx] && OptionNode->Tests[SubIdx]->NodeInstance)
 				{
 					AllNodes.Add(OptionNode->Tests[SubIdx]->NodeInstance);
 				}
