Made EQS avaiable as experimental feature in the editor #UE4

- also fixed EQS editor crashing when adding or clicking on tests

diff --git a/Engine/Source/Editor/EnvironmentQueryEditor/Private/EnvironmentQueryFactory.cpp b/Engine/Source/Editor/EnvironmentQueryEditor/Private/EnvironmentQueryFactory.cpp
--- a/Engine/Source/Editor/EnvironmentQueryEditor/Private/EnvironmentQueryFactory.cpp
+++ b/Engine/Source/Editor/EnvironmentQueryEditor/Private/EnvironmentQueryFactory.cpp
@@ -20,6 +20,11 @@ UObject* UEnvironmentQueryFactory::FactoryCreateNew(UClass* Class,UObject* InPar
 
 bool UEnvironmentQueryFactory::CanCreateNew() const
 {
+	if (GetDefault<UEditorExperimentalSettings>()->bEQSEditor)
+	{
+		return true;
+	}
+
 	// Check ini to see if we should enable creation
 	bool bEnableEnvironmentQueryEd = false;
 	GConfig->GetBool(TEXT("EnvironmentQueryEd"), TEXT("EnableEnvironmentQueryEd"), bEnableEnvironmentQueryEd, GEngineIni);
diff --git a/Engine/Source/Editor/EnvironmentQueryEditor/Private/SGraphNode_EnvironmentQuery.cpp b/Engine/Source/Editor/EnvironmentQueryEditor/Private/SGraphNode_EnvironmentQuery.cpp
--- a/Engine/Source/Editor/EnvironmentQueryEditor/Private/SGraphNode_EnvironmentQuery.cpp
+++ b/Engine/Source/Editor/EnvironmentQueryEditor/Private/SGraphNode_EnvironmentQuery.cpp
@@ -100,7 +100,7 @@ void SGraphNode_EnvironmentQuery::AddTest(TSharedPtr<SGraphNode> TestWidget)
 FSlateColor SGraphNode_EnvironmentQuery::GetBorderBackgroundColor() const
 {
 	UEnvironmentQueryGraphNode_Test* TestNode = Cast<UEnvironmentQueryGraphNode_Test>(GraphNode);
-	const bool bSelectedSubNode = TestNode && TestNode->ParentNode && GetOwnerPanel()->SelectionManager.SelectedNodes.Contains(GraphNode);
+	const bool bSelectedSubNode = TestNode && TestNode->ParentNode && GetOwnerPanel().IsValid() && GetOwnerPanel()->SelectionManager.SelectedNodes.Contains(GraphNode);
 
 	return bSelectedSubNode ? EnvironmentQueryColors::NodeBorder::Selected :
 		EnvironmentQueryColors::NodeBorder::Default;
@@ -678,4 +678,17 @@ const FSlateBrush* SGraphNode_EnvironmentQuery::GetNameIcon() const
 	return FEditorStyle::GetBrush( TEXT("Graph.StateNode.Icon") );
 }
 
+void SGraphNode_EnvironmentQuery::SetOwner(const TSharedRef<SGraphPanel>& OwnerPanel)
+{
+	SGraphNode::SetOwner(OwnerPanel);
+
+	for (auto& ChildWidget : TestWidgets)
+	{
+		if (ChildWidget.IsValid())
+		{
+			ChildWidget->SetOwner(OwnerPanel);
+		}
+	}
+}
+
 #undef LOCTEXT_NAMESPACE
diff --git a/Engine/Source/Editor/EnvironmentQueryEditor/Private/SGraphNode_EnvironmentQuery.h b/Engine/Source/Editor/EnvironmentQueryEditor/Private/SGraphNode_EnvironmentQuery.h
--- a/Engine/Source/Editor/EnvironmentQueryEditor/Private/SGraphNode_EnvironmentQuery.h
+++ b/Engine/Source/Editor/EnvironmentQueryEditor/Private/SGraphNode_EnvironmentQuery.h
@@ -21,6 +21,7 @@ public:
 	virtual void OnDragLeave( const FDragDropEvent& DragDropEvent ) override;
 	virtual FReply OnMouseMove(const FGeometry& SenderGeometry, const FPointerEvent& MouseEvent) override;
 	virtual void Tick( const FGeometry& AllottedGeometry, const double InCurrentTime, const float InDeltaTime ) override;
+	virtual void SetOwner(const TSharedRef<SGraphPanel>& OwnerPanel) override;
 	// End of SGraphNode interface
 
 	/** handle mouse down on the node */
diff --git a/Engine/Source/Editor/UnrealEd/Classes/Settings/EditorExperimentalSettings.h b/Engine/Source/Editor/UnrealEd/Classes/Settings/EditorExperimentalSettings.h
--- a/Engine/Source/Editor/UnrealEd/Classes/Settings/EditorExperimentalSettings.h
+++ b/Engine/Source/Editor/UnrealEd/Classes/Settings/EditorExperimentalSettings.h
@@ -77,6 +77,9 @@ public:
 	UPROPERTY(EditAnywhere, config, Category = Cooking, meta = (DisplayName = "Cook On The Side (Run a cook on the fly server in the background of the editor)"))
 	bool bCookOnTheSide;
 
+	/** Enables Environment Queries editor */
+	UPROPERTY(EditAnywhere, config, Category = AI, meta = (DisplayName = "Environment Querying System"))
+	bool bEQSEditor;
 
 	/**
 	 * Returns an event delegate that is executed when a setting has changed.
diff --git a/Engine/Source/Editor/UnrealEd/Private/Editor.cpp b/Engine/Source/Editor/UnrealEd/Private/Editor.cpp
--- a/Engine/Source/Editor/UnrealEd/Private/Editor.cpp
+++ b/Engine/Source/Editor/UnrealEd/Private/Editor.cpp
@@ -584,7 +584,7 @@ void UEditorEngine::Init(IEngineLoop* InEngineLoop)
 
 		bool bEnvironmentQueryEditor = false;
 		GConfig->GetBool(TEXT("EnvironmentQueryEd"), TEXT("EnableEnvironmentQueryEd"), bEnvironmentQueryEditor, GEngineIni);
-		if (bEnvironmentQueryEditor)
+		if (bEnvironmentQueryEditor || GetDefault<UEditorExperimentalSettings>()->bEQSEditor)
 		{
 			FModuleManager::Get().LoadModule(TEXT("EnvironmentQueryEditor"));
 		}
diff --git a/Engine/Source/Editor/UnrealEd/Private/Settings/SettingsClasses.cpp b/Engine/Source/Editor/UnrealEd/Private/Settings/SettingsClasses.cpp
--- a/Engine/Source/Editor/UnrealEd/Private/Settings/SettingsClasses.cpp
+++ b/Engine/Source/Editor/UnrealEd/Private/Settings/SettingsClasses.cpp
@@ -56,6 +56,8 @@ UEditorExperimentalSettings::UEditorExperimentalSettings( const FObjectInitializ
 
 void UEditorExperimentalSettings::PostEditChangeProperty( struct FPropertyChangedEvent& PropertyChangedEvent )
 {
+	static const FName NAME_EQS = GET_MEMBER_NAME_CHECKED(UEditorExperimentalSettings, bEQSEditor);
+
 	Super::PostEditChangeProperty(PropertyChangedEvent);
 
 	const FName Name = (PropertyChangedEvent.Property != nullptr) ? PropertyChangedEvent.Property->GetFName() : NAME_None;
@@ -64,6 +66,13 @@ void UEditorExperimentalSettings::PostEditChangeProperty( struct FPropertyChange
 	{
 		EKeys::SetConsoleForGamepadLabels(ConsoleForGamepadLabels);
 	}
+	else if (Name == NAME_EQS)
+	{
+		if (bEQSEditor)
+		{
+			FModuleManager::Get().LoadModule(TEXT("EnvironmentQueryEditor"));
+		}
+	}
 
 	if (!FUnrealEdMisc::Get().IsDeletePreferences())
 	{
diff --git a/Engine/Source/Runtime/Launch/Private/LaunchEngineLoop.cpp b/Engine/Source/Runtime/Launch/Private/LaunchEngineLoop.cpp
--- a/Engine/Source/Runtime/Launch/Private/LaunchEngineLoop.cpp
+++ b/Engine/Source/Runtime/Launch/Private/LaunchEngineLoop.cpp
@@ -1810,7 +1810,11 @@ bool FEngineLoop::LoadStartupCoreModules()
 	// cooking needs this module too
 	bool bEnvironmentQueryEditor = false;
 	GConfig->GetBool(TEXT("EnvironmentQueryEd"), TEXT("EnableEnvironmentQueryEd"), bEnvironmentQueryEditor, GEngineIni);
-	if (bEnvironmentQueryEditor)
+	if (bEnvironmentQueryEditor 
+#if WITH_EDITOR
+		|| GetDefault<UEditorExperimentalSettings>()->bEQSEditor
+#endif // WITH_EDITOR
+		)
 	{
 		FModuleManager::Get().LoadModule(TEXT("EnvironmentQueryEditor"));
 	}
