Added a way to specify additional vertical offset to EQS items projected to geometry/navigation #UE4

diff --git a/Engine/Source/Editor/EnvironmentQueryEditor/Private/DetailCustomizations/EnvTraceDataCustomization.cpp b/Engine/Source/Editor/EnvironmentQueryEditor/Private/DetailCustomizations/EnvTraceDataCustomization.cpp
--- a/Engine/Source/Editor/EnvironmentQueryEditor/Private/DetailCustomizations/EnvTraceDataCustomization.cpp
+++ b/Engine/Source/Editor/EnvironmentQueryEditor/Private/DetailCustomizations/EnvTraceDataCustomization.cpp
@@ -90,6 +90,10 @@ void FEnvTraceDataCustomization::CustomizeChildren( TSharedRef<class IPropertyHa
 		.Visibility(TAttribute<EVisibility>::Create(TAttribute<EVisibility>::FGetter::CreateSP(this, &FEnvTraceDataCustomization::GetProjectionVisibility)));
 
 	// advanced props
+	TSharedPtr<IPropertyHandle> PropPostProjectionVerticalOffset = StructPropertyHandle->GetChildHandle(GET_MEMBER_NAME_CHECKED(FEnvTraceData, PostProjectionVerticalOffset));
+	StructBuilder.AddChildProperty(PropPostProjectionVerticalOffset.ToSharedRef())
+		.Visibility(TAttribute<EVisibility>::Create(TAttribute<EVisibility>::FGetter::CreateSP(this, &FEnvTraceDataCustomization::GetProjectionVisibility)));
+
 	TSharedPtr<IPropertyHandle> PropTraceComplex = StructPropertyHandle->GetChildHandle(GET_MEMBER_NAME_CHECKED(FEnvTraceData,bTraceComplex));
 	StructBuilder.AddChildProperty(PropTraceComplex.ToSharedRef())
 		.Visibility(TAttribute<EVisibility>::Create(TAttribute<EVisibility>::FGetter::CreateSP(this, &FEnvTraceDataCustomization::GetGeometryVisibility)));
diff --git a/Engine/Source/Runtime/AIModule/Classes/EnvironmentQuery/EnvQueryTypes.h b/Engine/Source/Runtime/AIModule/Classes/EnvironmentQuery/EnvQueryTypes.h
--- a/Engine/Source/Runtime/AIModule/Classes/EnvironmentQuery/EnvQueryTypes.h
+++ b/Engine/Source/Runtime/AIModule/Classes/EnvironmentQuery/EnvQueryTypes.h
@@ -344,6 +344,13 @@ struct AIMODULE_API FEnvTraceData
 	UPROPERTY(EditDefaultsOnly, Category=Trace, meta=(UIMin=0, ClampMin=0))
 	float ExtentZ;
 
+	/** this value will be added to resulting location's Z axis. Can be useful when 
+	 *	projecting points to navigation since navmesh is just an approximation of level 
+	 *	geometry and items may end up being under collide-able geometry which would 
+	 *	for example falsify visibility tests.*/
+	UPROPERTY(EditDefaultsOnly, Category=Trace)
+	float PostProjectionVerticalOffset;
+
 	/** geometry trace channel */
 	UPROPERTY(EditDefaultsOnly, Category=Trace)
 	TEnumAsByte<enum ETraceTypeQuery> TraceChannel;
@@ -966,7 +973,7 @@ public:
 namespace FEQSHelpers
 {
 #if WITH_RECAST
-	const ARecastNavMesh* FindNavMeshForQuery(FEnvQueryInstance& QueryInstance);
+	AIMODULE_API const ARecastNavMesh* FindNavMeshForQuery(FEnvQueryInstance& QueryInstance);
 #endif // WITH_RECAST
 }
 
diff --git a/Engine/Source/Runtime/AIModule/Private/EnvironmentQuery/Generators/EnvQueryGenerator_ProjectedPoints.cpp b/Engine/Source/Runtime/AIModule/Private/EnvironmentQuery/Generators/EnvQueryGenerator_ProjectedPoints.cpp
--- a/Engine/Source/Runtime/AIModule/Private/EnvironmentQuery/Generators/EnvQueryGenerator_ProjectedPoints.cpp
+++ b/Engine/Source/Runtime/AIModule/Private/EnvironmentQuery/Generators/EnvQueryGenerator_ProjectedPoints.cpp
@@ -13,6 +13,8 @@ UEnvQueryGenerator_ProjectedPoints::UEnvQueryGenerator_ProjectedPoints(const FOb
 
 void UEnvQueryGenerator_ProjectedPoints::ProjectAndFilterNavPoints(TArray<FVector>& Points, const ANavigationData* NavData) const
 {
+	const FVector VerticalOffset(0, 0, ProjectionData.PostProjectionVerticalOffset);
+
 	if (ProjectionData.TraceMode == EEnvQueryTrace::Navigation && NavData)
 	{
 		TSharedPtr<const FNavigationQueryFilter> NavigationFilter = UNavigationQueryFilter::GetQueryFilter(NavData, ProjectionData.NavigationFilter);	
@@ -48,7 +50,7 @@ void UEnvQueryGenerator_ProjectedPoints::ProjectAndFilterNavPoints(TArray<FVecto
 			}
 			else
 			{
-				Points[PointIndex] = Workload[PointIndex].OutLocation.Location;
+				Points[PointIndex] = Workload[PointIndex].OutLocation.Location + VerticalOffset;
 			}
 		}
 	}
