integrated AIDataProviders from dev branch

diff --git a/Engine/Source/Editor/DetailCustomizations/DetailCustomizations.Build.cs b/Engine/Source/Editor/DetailCustomizations/DetailCustomizations.Build.cs
--- a/Engine/Source/Editor/DetailCustomizations/DetailCustomizations.Build.cs
+++ b/Engine/Source/Editor/DetailCustomizations/DetailCustomizations.Build.cs
@@ -41,7 +41,8 @@ public class DetailCustomizations : ModuleRules
 				"RHI",
                 "MediaAssets",
                 "HardwareTargeting",
-				"SharedSettingsWidgets"
+				"SharedSettingsWidgets",
+                "AIModule"
 			}
 		);
 
diff --git a/Engine/Source/Editor/DetailCustomizations/Private/AIDataProviderValueDetails.cpp b/Engine/Source/Editor/DetailCustomizations/Private/AIDataProviderValueDetails.cpp
new file mode 100644
--- /dev/null
+++ b/Engine/Source/Editor/DetailCustomizations/Private/AIDataProviderValueDetails.cpp
@@ -0,0 +1,148 @@
+// Copyright 1998-2015 Epic Games, Inc. All Rights Reserved.
+
+#include "DetailCustomizationsPrivatePCH.h"
+#include "AIDataProviderValueDetails.h"
+#include "DataProviders/AIDataProvider.h"
+
+#define LOCTEXT_NAMESPACE "AIDataProviderValueDetails"
+
+TSharedRef<IPropertyTypeCustomization> FAIDataProviderValueDetails::MakeInstance()
+{
+	return MakeShareable(new FAIDataProviderValueDetails);
+}
+
+BEGIN_SLATE_FUNCTION_BUILD_OPTIMIZATION
+void FAIDataProviderValueDetails::CustomizeHeader(TSharedRef<class IPropertyHandle> StructPropertyHandle, class FDetailWidgetRow& HeaderRow, IPropertyTypeCustomizationUtils& StructCustomizationUtils)
+{
+	DataBindingProperty = StructPropertyHandle->GetChildHandle(GET_MEMBER_NAME_CHECKED(FAIDataProviderValue, DataBinding));
+	DataFieldProperty = StructPropertyHandle->GetChildHandle(GET_MEMBER_NAME_CHECKED(FAIDataProviderValue, DataField));
+	DefaultValueProperty = StructPropertyHandle->GetChildHandle(TEXT("DefaultValue"));
+
+	DataBindingProperty->SetOnPropertyValueChanged(FSimpleDelegate::CreateSP(this, &FAIDataProviderValueDetails::OnBindingChanged));
+	
+	TArray<void*> StructPtrs;
+	StructPropertyHandle->AccessRawData(StructPtrs);
+	DataPtr = (StructPtrs.Num() == 1) ? reinterpret_cast<FAIDataProviderValue*>(StructPtrs[0]) : nullptr;
+	OnBindingChanged();
+
+	TSharedRef<SWidget> DefPropWidget = DefaultValueProperty->CreatePropertyValueWidget();
+	DefPropWidget->SetVisibility(TAttribute<EVisibility>::Create(TAttribute<EVisibility>::FGetter::CreateSP(this, &FAIDataProviderValueDetails::GetDefaultValueVisibility)));
+
+	HeaderRow
+	.NameContent()
+	[
+		StructPropertyHandle->CreatePropertyNameWidget()
+	]
+	.ValueContent()
+	.VAlign(VAlign_Center)
+	.MinDesiredWidth(300.0f)
+	[
+		SNew(SHorizontalBox)
+		+ SHorizontalBox::Slot()
+		.Padding(0.0f, 2.0f, 5.0f, 2.0f)
+		[
+			SNew(STextBlock)
+			.Text(this, &FAIDataProviderValueDetails::GetValueDesc)
+			.Font(IDetailLayoutBuilder::GetDetailFont())
+			.Visibility(TAttribute<EVisibility>::Create(TAttribute<EVisibility>::FGetter::CreateSP(this, &FAIDataProviderValueDetails::GetBindingDescVisibility)))
+		]
+		+ SHorizontalBox::Slot()
+		.Padding(0.0f, 2.0f, 5.0f, 2.0f)
+		[
+			DefPropWidget
+		]
+	];
+}
+END_SLATE_FUNCTION_BUILD_OPTIMIZATION
+
+void FAIDataProviderValueDetails::CustomizeChildren(TSharedRef<class IPropertyHandle> StructPropertyHandle, class IDetailChildrenBuilder& StructBuilder, IPropertyTypeCustomizationUtils& StructCustomizationUtils)
+{
+	if (StructPropertyHandle->IsValidHandle())
+	{
+		StructBuilder.AddChildProperty(DataBindingProperty.ToSharedRef());
+
+		StructBuilder.AddChildContent(LOCTEXT("PropertyField", "Property"))
+		.Visibility(TAttribute<EVisibility>::Create(TAttribute<EVisibility>::FGetter::CreateSP(this, &FAIDataProviderValueDetails::GetDataFieldVisibility)))
+		.NameContent()
+		[
+			DataFieldProperty->CreatePropertyNameWidget()
+		]
+		.ValueContent()
+		[
+			SNew(SComboButton)
+			.OnGetMenuContent(this, &FAIDataProviderValueDetails::OnGetDataFieldContent)
+			.ContentPadding(FMargin(2.0f, 2.0f))
+			.ButtonContent()
+			[
+				SNew(STextBlock)
+				.Text(this, &FAIDataProviderValueDetails::GetDataFieldDesc)
+				.Font(IDetailLayoutBuilder::GetDetailFont())
+			]
+		];		
+	}
+}
+
+void FAIDataProviderValueDetails::OnBindingChanged()
+{
+	MatchingProperties.Reset();
+	if (DataPtr)
+	{
+		DataPtr->GetMatchingProperties(MatchingProperties);
+
+		FName NameValue;
+		DataFieldProperty->GetValue(NameValue);
+		if (MatchingProperties.Num() && !MatchingProperties.Contains(NameValue))
+		{
+			OnDataFieldNameChange(0);
+		}
+	}
+}
+
+TSharedRef<SWidget> FAIDataProviderValueDetails::OnGetDataFieldContent()
+{
+	FMenuBuilder MenuBuilder(true, NULL);
+
+	for (int32 i = 0; i < MatchingProperties.Num(); i++)
+	{
+		FUIAction ItemAction(FExecuteAction::CreateSP(this, &FAIDataProviderValueDetails::OnDataFieldNameChange, i));
+		MenuBuilder.AddMenuEntry(FText::FromName(MatchingProperties[i]), TAttribute<FText>(), FSlateIcon(), ItemAction);
+	}
+
+	return MenuBuilder.MakeWidget();
+}
+
+void FAIDataProviderValueDetails::OnDataFieldNameChange(int32 Index)
+{
+	FName NameValue = MatchingProperties[Index];
+	DataFieldProperty->SetValue(NameValue);
+}
+
+FString FAIDataProviderValueDetails::GetDataFieldDesc() const
+{
+	FName NameValue;
+	DataFieldProperty->GetValue(NameValue);
+
+	return NameValue.ToString();
+}
+
+FString FAIDataProviderValueDetails::GetValueDesc() const
+{
+	return DataPtr ? DataPtr->ToString() : TEXT("empty");
+}
+
+EVisibility FAIDataProviderValueDetails::GetDataFieldVisibility() const
+{
+	return (DataPtr && DataPtr->DataBinding && (MatchingProperties.Num() > 1)) ? EVisibility::Visible : EVisibility::Collapsed;
+}
+
+EVisibility FAIDataProviderValueDetails::GetBindingDescVisibility() const
+{
+	return (DataPtr && DataPtr->DataBinding) ? EVisibility::Visible : EVisibility::Collapsed;
+}
+
+EVisibility FAIDataProviderValueDetails::GetDefaultValueVisibility() const
+{
+	return (!DataPtr || !DataPtr->DataBinding) ? EVisibility::Visible : EVisibility::Collapsed;
+}
+
+#undef LOCTEXT_NAMESPACE
diff --git a/Engine/Source/Editor/DetailCustomizations/Private/AIDataProviderValueDetails.h b/Engine/Source/Editor/DetailCustomizations/Private/AIDataProviderValueDetails.h
new file mode 100644
--- /dev/null
+++ b/Engine/Source/Editor/DetailCustomizations/Private/AIDataProviderValueDetails.h
@@ -0,0 +1,34 @@
+// Copyright 1998-2015 Epic Games, Inc. All Rights Reserved.
+
+#pragma once
+
+struct FAIDataProviderValue;
+
+class FAIDataProviderValueDetails : public IPropertyTypeCustomization
+{
+public:
+	/** Makes a new instance of this detail layout class for a specific detail view requesting it */
+	static TSharedRef<IPropertyTypeCustomization> MakeInstance();
+
+	/** IPropertyTypeCustomization interface */
+	virtual void CustomizeHeader(TSharedRef<class IPropertyHandle> StructPropertyHandle, class FDetailWidgetRow& HeaderRow, IPropertyTypeCustomizationUtils& StructCustomizationUtils) override;
+	virtual void CustomizeChildren(TSharedRef<class IPropertyHandle> StructPropertyHandle, class IDetailChildrenBuilder& StructBuilder, IPropertyTypeCustomizationUtils& StructCustomizationUtils) override;
+
+private:
+
+	TSharedPtr<IPropertyHandle> DataBindingProperty;
+	TSharedPtr<IPropertyHandle> DataFieldProperty;
+	TSharedPtr<IPropertyHandle> DefaultValueProperty;
+	FAIDataProviderValue* DataPtr;
+
+	TArray<FName> MatchingProperties;
+
+	void OnBindingChanged();
+	TSharedRef<SWidget> OnGetDataFieldContent();
+	void OnDataFieldNameChange(int32 Index);
+	FString GetDataFieldDesc() const;
+	FString GetValueDesc() const;
+	EVisibility GetBindingDescVisibility() const;
+	EVisibility GetDataFieldVisibility() const;
+	EVisibility GetDefaultValueVisibility() const;
+};
diff --git a/Engine/Source/Editor/DetailCustomizations/Private/DetailCustomizations.cpp b/Engine/Source/Editor/DetailCustomizations/Private/DetailCustomizations.cpp
--- a/Engine/Source/Editor/DetailCustomizations/Private/DetailCustomizations.cpp
+++ b/Engine/Source/Editor/DetailCustomizations/Private/DetailCustomizations.cpp
@@ -58,6 +58,7 @@
 #include "ParticleModuleDetails.h"
 #include "CameraDetails.h"
 #include "BlackboardEntryDetails.h"
+#include "AIDataProviderValueDetails.h"
 #include "EnvQueryParamInstanceCustomization.h"
 #include "SkeletonNotifyDetails.h"
 #include "SlateColorCustomization.h"
@@ -171,6 +172,9 @@ void FDetailCustomizationsModule::RegisterPropertyTypeCustomizations()
 	RegisterCustomPropertyTypeLayout("DateTime", FOnGetPropertyTypeCustomizationInstance::CreateStatic(&FDateTimeStructCustomization::MakeInstance));
 	RegisterCustomPropertyTypeLayout("Timespan", FOnGetPropertyTypeCustomizationInstance::CreateStatic(&FTimespanStructCustomization::MakeInstance));
 	RegisterCustomPropertyTypeLayout("BlackboardEntry", FOnGetPropertyTypeCustomizationInstance::CreateStatic(&FBlackboardEntryDetails::MakeInstance));
+	RegisterCustomPropertyTypeLayout("AIDataProviderIntValue", FOnGetPropertyTypeCustomizationInstance::CreateStatic(&FAIDataProviderValueDetails::MakeInstance));
+	RegisterCustomPropertyTypeLayout("AIDataProviderFloatValue", FOnGetPropertyTypeCustomizationInstance::CreateStatic(&FAIDataProviderValueDetails::MakeInstance));
+	RegisterCustomPropertyTypeLayout("AIDataProviderBoolValue", FOnGetPropertyTypeCustomizationInstance::CreateStatic(&FAIDataProviderValueDetails::MakeInstance));
 	RegisterCustomPropertyTypeLayout("RuntimeFloatCurve", FOnGetPropertyTypeCustomizationInstance::CreateStatic(&FCurveStructCustomization::MakeInstance));
 	RegisterCustomPropertyTypeLayout("EnvNamedValue", FOnGetPropertyTypeCustomizationInstance::CreateStatic(&FEnvQueryParamInstanceCustomization::MakeInstance));
 	RegisterCustomPropertyTypeLayout("NavigationLink", FOnGetPropertyTypeCustomizationInstance::CreateStatic(&FNavLinkStructCustomization::MakeInstance));
diff --git a/Engine/Source/Editor/EnvironmentQueryEditor/Private/DetailCustomizations/EnvQueryTestDetails.cpp b/Engine/Source/Editor/EnvironmentQueryEditor/Private/DetailCustomizations/EnvQueryTestDetails.cpp
--- a/Engine/Source/Editor/EnvironmentQueryEditor/Private/DetailCustomizations/EnvQueryTestDetails.cpp
+++ b/Engine/Source/Editor/EnvironmentQueryEditor/Private/DetailCustomizations/EnvQueryTestDetails.cpp
@@ -30,61 +30,25 @@ void FEnvQueryTestDetails::CustomizeDetails( IDetailLayoutBuilder& DetailLayout
 	}
 
 	// Initialize all handles
-	ConditionHandle = DetailLayout.GetProperty(GET_MEMBER_NAME_CHECKED(UEnvQueryTest, Condition));
 	FilterTypeHandle = DetailLayout.GetProperty(GET_MEMBER_NAME_CHECKED(UEnvQueryTest, FilterType));
 	ScoreEquationHandle = DetailLayout.GetProperty(GET_MEMBER_NAME_CHECKED(UEnvQueryTest, ScoringEquation));
 	TestPurposeHandle = DetailLayout.GetProperty(GET_MEMBER_NAME_CHECKED(UEnvQueryTest, TestPurpose));
-	ScoreHandle = DetailLayout.GetProperty(GET_MEMBER_NAME_CHECKED(UEnvQueryTest, Weight));
+	ScoreHandle = DetailLayout.GetProperty(GET_MEMBER_NAME_CHECKED(UEnvQueryTest, ScoringFactor));
 
 	ClampMinTypeHandle = DetailLayout.GetProperty(GET_MEMBER_NAME_CHECKED(UEnvQueryTest, ClampMinType));
 	ClampMaxTypeHandle = DetailLayout.GetProperty(GET_MEMBER_NAME_CHECKED(UEnvQueryTest, ClampMaxType));
 
-	ScoreClampingMinHandle = DetailLayout.GetProperty(GET_MEMBER_NAME_CHECKED(UEnvQueryTest, ScoreClampingMin));
-	FloatFilterMinHandle = DetailLayout.GetProperty(GET_MEMBER_NAME_CHECKED(UEnvQueryTest, FloatFilterMin));
+	ScoreClampMinHandle = DetailLayout.GetProperty(GET_MEMBER_NAME_CHECKED(UEnvQueryTest, ScoreClampMin));
+	FloatValueMinHandle = DetailLayout.GetProperty(GET_MEMBER_NAME_CHECKED(UEnvQueryTest, FloatValueMin));
 
-	ScoreClampingMaxHandle = DetailLayout.GetProperty(GET_MEMBER_NAME_CHECKED(UEnvQueryTest, ScoreClampingMax));
-	FloatFilterMaxHandle = DetailLayout.GetProperty(GET_MEMBER_NAME_CHECKED(UEnvQueryTest, FloatFilterMax));
+	ScoreClampMaxHandle = DetailLayout.GetProperty(GET_MEMBER_NAME_CHECKED(UEnvQueryTest, ScoreClampMax));
+	FloatValueMaxHandle = DetailLayout.GetProperty(GET_MEMBER_NAME_CHECKED(UEnvQueryTest, FloatValueMax));
 
 	// Build combo box data values
-	BuildConditionValues();
 	BuildFilterTestValues();
 	BuildScoreEquationValues();
 	BuildScoreClampingTypeValues(true, ClampMinTypeValues);
 	BuildScoreClampingTypeValues(false, ClampMaxTypeValues);
-
-// 	// dynamic Condition combo
-// 	IDetailCategoryBuilder& DeprecatedFilterCategory = DetailLayout.EditCategory("Deprecated Filter and Score");
-// 	IDetailPropertyRow& ConditionRow = DeprecatedFilterCategory.AddProperty(ConditionHandle);
-// 	ConditionRow.CustomWidget()
-// 		.NameContent()
-// 		[
-// 			ConditionHandle->CreatePropertyNameWidget()
-// 		]
-// 		.ValueContent()
-// 		[
-// 			SNew(SComboButton)
-// 			.OnGetMenuContent(this, &FEnvQueryTestDetails::OnGetConditionContent)
-// 			.ContentPadding(FMargin( 2.0f, 2.0f ))
-// 			.ButtonContent()
-// 			[
-// 				SNew(STextBlock) 
-// 				.Text(this, &FEnvQueryTestDetails::GetCurrentConditionDesc)
-// 				.Font(IDetailLayoutBuilder::GetDetailFont())
-// 			]
-// 		];
-// 	ConditionRow.EditCondition(TAttribute<bool>(this, &FEnvQueryTestDetails::AllowWritingToFiltersFromScore), NULL);
-// 
-// 	// TODO: Remove!
-// 	IDetailPropertyRow& DiscardFailedRow = DeprecatedFilterCategory.AddProperty(DetailLayout.GetProperty(GET_MEMBER_NAME_CHECKED(UEnvQueryTest, bDiscardFailedItems)));
-// 	DiscardFailedRow.EditCondition(TAttribute<bool>(this, &FEnvQueryTestDetails::AllowWritingToFiltersFromScore), NULL);
-// 
-// 	// DEPRECATED!  Remove as soon as testing is complete!
-// 	IDetailPropertyRow& FloatFilterRow = DeprecatedFilterCategory.AddProperty(DetailLayout.GetProperty(GET_MEMBER_NAME_CHECKED(UEnvQueryTest, FloatFilter)));
-// 	FloatFilterRow.Visibility(TAttribute<EVisibility>::Create(TAttribute<EVisibility>::FGetter::CreateSP(this, &FEnvQueryTestDetails::GetFloatFilterVisibility)));
-// 	FloatFilterRow.EditCondition(TAttribute<bool>(this, &FEnvQueryTestDetails::AllowWritingToFiltersFromScore), NULL);
-// 
-// 	IDetailPropertyRow& WeightModiferRow = DeprecatedFilterCategory.AddProperty(DetailLayout.GetProperty(GET_MEMBER_NAME_CHECKED(UEnvQueryTest, WeightModifier)));
-// 	WeightModiferRow.EditCondition(TAttribute<bool>(this, &FEnvQueryTestDetails::AllowWritingToFiltersFromScore), NULL);
 	
 	IDetailCategoryBuilder& TestCategory = DetailLayout.EditCategory("Test");
 	IDetailPropertyRow& TestPurposeRow = TestCategory.AddProperty(TestPurposeHandle);
@@ -112,16 +76,18 @@ void FEnvQueryTestDetails::CustomizeDetails( IDetailLayoutBuilder& DetailLayout
   	FilterTypeRow.Visibility(TAttribute<EVisibility>::Create(TAttribute<EVisibility>::FGetter::CreateSP(this, &FEnvQueryTestDetails::GetFloatFilterVisibility)));
 
 	// filters
-	IDetailPropertyRow& FloatFilterMinRow = FilterCategory.AddProperty(DetailLayout.GetProperty(GET_MEMBER_NAME_CHECKED(UEnvQueryTest, FloatFilterMin)));
-	FloatFilterMinRow.Visibility(TAttribute<EVisibility>::Create(TAttribute<EVisibility>::FGetter::CreateSP(this, &FEnvQueryTestDetails::GetVisibilityOfFloatFilterMin)));
+	IDetailPropertyRow& FloatValueMinRow = FilterCategory.AddProperty(DetailLayout.GetProperty(GET_MEMBER_NAME_CHECKED(UEnvQueryTest, FloatValueMin)));
+	FloatValueMinRow.Visibility(TAttribute<EVisibility>::Create(TAttribute<EVisibility>::FGetter::CreateSP(this, &FEnvQueryTestDetails::GetVisibilityOfFloatValueMin)));
 
-	IDetailPropertyRow& FloatFilterMaxRow = FilterCategory.AddProperty(DetailLayout.GetProperty(GET_MEMBER_NAME_CHECKED(UEnvQueryTest, FloatFilterMax)));
-	FloatFilterMaxRow.Visibility(TAttribute<EVisibility>::Create(TAttribute<EVisibility>::FGetter::CreateSP(this, &FEnvQueryTestDetails::GetVisibilityOfFloatFilterMax)));
+	IDetailPropertyRow& FloatValueMaxRow = FilterCategory.AddProperty(DetailLayout.GetProperty(GET_MEMBER_NAME_CHECKED(UEnvQueryTest, FloatValueMax)));
+	FloatValueMaxRow.Visibility(TAttribute<EVisibility>::Create(TAttribute<EVisibility>::FGetter::CreateSP(this, &FEnvQueryTestDetails::GetVisibilityOfFloatValueMax)));
 
-	IDetailPropertyRow& BoolFilterRow = FilterCategory.AddProperty(DetailLayout.GetProperty(GET_MEMBER_NAME_CHECKED(UEnvQueryTest, BoolFilter)));
-	BoolFilterRow.Visibility(TAttribute<EVisibility>::Create(TAttribute<EVisibility>::FGetter::CreateSP(this, &FEnvQueryTestDetails::GetBoolFilterVisibility)));
+	IDetailPropertyRow& BoolValueRow = FilterCategory.AddProperty(DetailLayout.GetProperty(GET_MEMBER_NAME_CHECKED(UEnvQueryTest, BoolValue)));
+	BoolValueRow.Visibility(TAttribute<EVisibility>::Create(TAttribute<EVisibility>::FGetter::CreateSP(this, &FEnvQueryTestDetails::GetBoolValueVisibility)));
 
+	// required when it's created for "scoring only" tests
 	IDetailGroup& HackToEnsureFilterCategoryIsVisible = FilterCategory.AddGroup("HackForVisibility", FText::GetEmpty());
+	HackToEnsureFilterCategoryIsVisible.HeaderRow().Visibility(EVisibility::Hidden);
 
 	// Scoring
 	IDetailCategoryBuilder& ScoreCategory = DetailLayout.EditCategory("Score");
@@ -129,6 +95,13 @@ void FEnvQueryTestDetails::CustomizeDetails( IDetailLayoutBuilder& DetailLayout
 	//----------------------------
 	// BEGIN Scoring: Clamping
 	IDetailGroup& ClampingGroup = ScoreCategory.AddGroup("Clamping", LOCTEXT("ClampingLabel", "Clamping"));
+	ClampingGroup.HeaderRow()
+	[
+		SNew(STextBlock)
+		.Visibility(TAttribute<EVisibility>::Create(TAttribute<EVisibility>::FGetter::CreateSP(this, &FEnvQueryTestDetails::GetClampingVisibility)))
+		.Text(FText::FromString("Clamping"))
+		.Font(IDetailLayoutBuilder::GetDetailFont())
+	];
 	
 	// Drop-downs for setting type of lower and upper bound normalization
 	IDetailPropertyRow& ClampMinTypeRow = ClampingGroup.AddPropertyRow(ClampMinTypeHandle.ToSharedRef());
@@ -152,16 +125,16 @@ void FEnvQueryTestDetails::CustomizeDetails( IDetailLayoutBuilder& DetailLayout
 	ClampMinTypeRow.Visibility(TAttribute<EVisibility>::Create(TAttribute<EVisibility>::FGetter::CreateSP(this, &FEnvQueryTestDetails::GetFloatScoreVisibility)));
 
 	// Lower Bound for normalization of score if specified independently of filtering.
-	IDetailPropertyRow& ScoreClampingMinRow = ClampingGroup.AddPropertyRow(ScoreClampingMinHandle.ToSharedRef());
-	ScoreClampingMinRow.Visibility(TAttribute<EVisibility>::Create(TAttribute<EVisibility>::FGetter::CreateSP(this, &FEnvQueryTestDetails::GetVisibilityOfScoreClampingMinimum)));
+	IDetailPropertyRow& ScoreClampMinRow = ClampingGroup.AddPropertyRow(ScoreClampMinHandle.ToSharedRef());
+	ScoreClampMinRow.Visibility(TAttribute<EVisibility>::Create(TAttribute<EVisibility>::FGetter::CreateSP(this, &FEnvQueryTestDetails::GetVisibilityOfScoreClampMinimum)));
 
 	// Lower Bound for scoring when tied to filter minimum.
-	if (FloatFilterMinHandle->IsValidHandle())
+	if (FloatValueMinHandle->IsValidHandle())
 	{
-		IDetailPropertyRow& FloatFilterMinForClampingRow = ClampingGroup.AddPropertyRow(FloatFilterMinHandle.ToSharedRef());
-		FloatFilterMinForClampingRow.Visibility(TAttribute<EVisibility>::Create(TAttribute<EVisibility>::FGetter::CreateSP(this, &FEnvQueryTestDetails::GetVisibilityOfFilterMinForScoreClamping)));
-		FloatFilterMinForClampingRow.ToolTip(LOCTEXT("FloatFilterMinForClampingRowToolTip", "See Filter Thresholds under the Filter tab.  Values lower than this (before clamping) cause the item to be thrown out as invalid.  Values are normalized with this value as the minimum, so items with this value will have a normalized score of 0."));
-		FloatFilterMinForClampingRow.EditCondition(TAttribute<bool>(this, &FEnvQueryTestDetails::AllowWritingToFiltersFromScore), NULL);
+		IDetailPropertyRow& FloatValueMinForClampingRow = ClampingGroup.AddPropertyRow(FloatValueMinHandle.ToSharedRef());
+		FloatValueMinForClampingRow.Visibility(TAttribute<EVisibility>::Create(TAttribute<EVisibility>::FGetter::CreateSP(this, &FEnvQueryTestDetails::GetVisibilityOfValueMinForScoreClamping)));
+		FloatValueMinForClampingRow.ToolTip(LOCTEXT("FloatFilterMinForClampingRowToolTip", "See Filter Thresholds under the Filter tab.  Values lower than this (before clamping) cause the item to be thrown out as invalid.  Values are normalized with this value as the minimum, so items with this value will have a normalized score of 0."));
+		FloatValueMinForClampingRow.EditCondition(TAttribute<bool>(this, &FEnvQueryTestDetails::AllowWritingToFiltersFromScore), NULL);
 	}
 
 	if (ClampMaxTypeHandle->IsValidHandle())
@@ -188,27 +161,28 @@ void FEnvQueryTestDetails::CustomizeDetails( IDetailLayoutBuilder& DetailLayout
 	}
 
 	// Upper Bound for normalization of score if specified independently of filtering.
-	if (ScoreClampingMaxHandle->IsValidHandle())
+	if (ScoreClampMaxHandle->IsValidHandle())
 	{
-		IDetailPropertyRow& ScoreClampingMaxRow = ClampingGroup.AddPropertyRow(ScoreClampingMaxHandle.ToSharedRef());
-		ScoreClampingMaxRow.Visibility(TAttribute<EVisibility>::Create(TAttribute<EVisibility>::FGetter::CreateSP(this, &FEnvQueryTestDetails::GetVisibilityOfScoreClampingMaximum)));
+		IDetailPropertyRow& ScoreClampMaxRow = ClampingGroup.AddPropertyRow(ScoreClampMaxHandle.ToSharedRef());
+		ScoreClampMaxRow.Visibility(TAttribute<EVisibility>::Create(TAttribute<EVisibility>::FGetter::CreateSP(this, &FEnvQueryTestDetails::GetVisibilityOfScoreClampMaximum)));
 	}
 	
-	if (FloatFilterMaxHandle->IsValidHandle())
+	if (FloatValueMaxHandle->IsValidHandle())
 	{
 		// Upper Bound for scoring when tied to filter maximum.
-		IDetailPropertyRow& FloatFilterMaxForClampingRow = ClampingGroup.AddPropertyRow(FloatFilterMaxHandle.ToSharedRef());
-		FloatFilterMaxForClampingRow.Visibility(TAttribute<EVisibility>::Create(TAttribute<EVisibility>::FGetter::CreateSP(this, &FEnvQueryTestDetails::GetVisibilityOfFilterMaxForScoreClamping)));
-		FloatFilterMaxForClampingRow.ToolTip(LOCTEXT("FloatFilterMaxForClampingRowToolTip", "See Filter Thresholds under the Filter tab.  Values higher than this (before normalization) cause the item to be thrown out as invalid.  Values are normalized with this value as the maximum, so items with this value will have a normalized score of 1."));
-		FloatFilterMaxForClampingRow.EditCondition(TAttribute<bool>(this, &FEnvQueryTestDetails::AllowWritingToFiltersFromScore), NULL);
+		IDetailPropertyRow& FloatValueMaxForClampingRow = ClampingGroup.AddPropertyRow(FloatValueMaxHandle.ToSharedRef());
+		FloatValueMaxForClampingRow.Visibility(TAttribute<EVisibility>::Create(TAttribute<EVisibility>::FGetter::CreateSP(this, &FEnvQueryTestDetails::GetVisibilityOfValueMaxForScoreClamping)));
+		FloatValueMaxForClampingRow.ToolTip(LOCTEXT("FloatFilterMaxForClampingRowToolTip", "See Filter Thresholds under the Filter tab.  Values higher than this (before normalization) cause the item to be thrown out as invalid.  Values are normalized with this value as the maximum, so items with this value will have a normalized score of 1."));
+		FloatValueMaxForClampingRow.EditCondition(TAttribute<bool>(this, &FEnvQueryTestDetails::AllowWritingToFiltersFromScore), NULL);
 	}
+
 	// END Scoring: Clamping, continue Scoring
 	//----------------------------
 
-	IDetailPropertyRow& BoolScoreTestRow = ScoreCategory.AddProperty(DetailLayout.GetProperty(GET_MEMBER_NAME_CHECKED(UEnvQueryTest, BoolFilter)));
-	BoolScoreTestRow.Visibility(TAttribute<EVisibility>::Create(TAttribute<EVisibility>::FGetter::CreateSP(this, &FEnvQueryTestDetails::GetBoolFilterVisibilityForScoring)));
+	IDetailPropertyRow& BoolScoreTestRow = ScoreCategory.AddProperty(DetailLayout.GetProperty(GET_MEMBER_NAME_CHECKED(UEnvQueryTest, BoolValue)));
+	BoolScoreTestRow.Visibility(TAttribute<EVisibility>::Create(TAttribute<EVisibility>::FGetter::CreateSP(this, &FEnvQueryTestDetails::GetBoolValueVisibilityForScoring)));
 	BoolScoreTestRow.DisplayName(LOCTEXT("BoolMatchLabel", "Bool Match"));
-	BoolScoreTestRow.ToolTip(LOCTEXT("BoolMatchToolTip", "Boolean value to match in order to grant score of 'Weight'.  Not matching this value will not change score."));
+	BoolScoreTestRow.ToolTip(LOCTEXT("BoolMatchToolTip", "Boolean value to match in order to grant score of 'ScoringFactor'.  Not matching this value will not change score."));
 
 // 	IDetailPropertyRow& ScoreMirrorNormalizedScoreRow = ScoreCategory.AddProperty(DetailLayout.GetProperty(GET_MEMBER_NAME_CHECKED(UEnvQueryTest, bMirrorNormalizedScore)));
 // 	ScoreMirrorNormalizedScoreRow.Visibility(TAttribute<EVisibility>::Create(TAttribute<EVisibility>::FGetter::CreateSP(this, &FEnvQueryTestDetails::GetFloatScoreVisibility)));
@@ -249,8 +223,8 @@ void FEnvQueryTestDetails::CustomizeDetails( IDetailLayoutBuilder& DetailLayout
 		];
 	ScoreEquationTypeRow.Visibility(TAttribute<EVisibility>::Create(TAttribute<EVisibility>::FGetter::CreateSP(this, &FEnvQueryTestDetails::GetFloatScoreVisibility)));
 
-	IDetailPropertyRow& ScoreWeightRow = ScoreCategory.AddProperty(DetailLayout.GetProperty(GET_MEMBER_NAME_CHECKED(UEnvQueryTest, Weight)));
-	ScoreWeightRow.Visibility(TAttribute<EVisibility>::Create(TAttribute<EVisibility>::FGetter::CreateSP(this, &FEnvQueryTestDetails::GetScoreVisibility)));
+	IDetailPropertyRow& ScoreFactorRow = ScoreCategory.AddProperty(DetailLayout.GetProperty(GET_MEMBER_NAME_CHECKED(UEnvQueryTest, ScoringFactor)));
+	ScoreFactorRow.Visibility(TAttribute<EVisibility>::Create(TAttribute<EVisibility>::FGetter::CreateSP(this, &FEnvQueryTestDetails::GetScoreVisibility)));
 
 	// scoring & filter function preview
 	IDetailCategoryBuilder& PreviewCategory = DetailLayout.EditCategory("Preview");
@@ -316,41 +290,12 @@ void FEnvQueryTestDetails::BuildScoreEquationValues()
 	}
 }
 
-void FEnvQueryTestDetails::BuildConditionValues()
-{
-	UEnum* TestConditionEnum = FindObject<UEnum>(ANY_PACKAGE, TEXT("EEnvTestCondition"));
-	check(TestConditionEnum);
-
-	ConditionValues.Reset();
-	ConditionValues.Add(FStringIntPair(TestConditionEnum->GetEnumText(EEnvTestCondition::NoCondition).ToString(), EEnvTestCondition::NoCondition));
-
-	const UEnvQueryTest* EditedTest = Cast<const UEnvQueryTest>(MyTest.Get());
-	if (EditedTest)
-	{
-		if (EditedTest->GetWorkOnFloatValues())
-		{
-			ConditionValues.Add(FStringIntPair(TestConditionEnum->GetEnumText(EEnvTestCondition::AtLeast).ToString(), EEnvTestCondition::AtLeast));
-			ConditionValues.Add(FStringIntPair(TestConditionEnum->GetEnumText(EEnvTestCondition::UpTo).ToString(), EEnvTestCondition::UpTo));
-		}
-		else
-		{
-			ConditionValues.Add(FStringIntPair(TestConditionEnum->GetEnumText(EEnvTestCondition::Match).ToString(), EEnvTestCondition::Match));
-		}
-	}
-}
-
 void FEnvQueryTestDetails::OnFilterTestChange(int32 Index)
 {
 	uint8 EnumValue = Index;
 	FilterTypeHandle->SetValue(EnumValue);
 }
 
-void FEnvQueryTestDetails::OnConditionComboChange(int32 Index)
-{
-	uint8 EnumValue = Index;
-	ConditionHandle->SetValue(EnumValue);
-}
-
 void FEnvQueryTestDetails::OnScoreEquationChange(int32 Index)
 {
 	uint8 EnumValue = Index;
@@ -511,20 +456,6 @@ FString FEnvQueryTestDetails::GetCurrentFilterTestDesc() const
 	return FString();
 }
 
-TSharedRef<SWidget> FEnvQueryTestDetails::OnGetConditionContent()
-{
-	BuildConditionValues();
-	FMenuBuilder MenuBuilder(true, NULL);
-
-	for (int32 i = 0; i < ConditionValues.Num(); i++)
-	{
-		FUIAction ItemAction( FExecuteAction::CreateSP( this, &FEnvQueryTestDetails::OnConditionComboChange, ConditionValues[i].Int ) );
-		MenuBuilder.AddMenuEntry( FText::FromString( ConditionValues[i].Str ), TAttribute<FText>(), FSlateIcon(), ItemAction);
-	}
-
-	return MenuBuilder.MakeWidget();
-}
-
 TSharedRef<SWidget> FEnvQueryTestDetails::OnGetEquationValuesContent()
 {
 	BuildScoreEquationValues();
@@ -539,22 +470,6 @@ TSharedRef<SWidget> FEnvQueryTestDetails::OnGetEquationValuesContent()
 	return MenuBuilder.MakeWidget();
 }
 
-FString FEnvQueryTestDetails::GetCurrentConditionDesc() const
-{
-	uint8 EnumValue;
-	ConditionHandle->GetValue(EnumValue);
-
-	for (int32 i = 0; i < ConditionValues.Num(); i++)
-	{
-		if (ConditionValues[i].Int == EnumValue)
-		{
-			return ConditionValues[i].Str;
-		}
-	}
-
-	return FString();
-}
-
 FString FEnvQueryTestDetails::GetScoreEquationInfo() const
 {
 	uint8 EnumValue;
@@ -563,16 +478,16 @@ FString FEnvQueryTestDetails::GetScoreEquationInfo() const
 	switch (EnumValue)
 	{
 		case EEnvTestScoreEquation::Linear:
-			return LOCTEXT("Linear","Final score = Weight * NormalizedItemValue").ToString();
+			return LOCTEXT("Linear","Final score = ScoringFactor * NormalizedItemValue").ToString();
 
 		case EEnvTestScoreEquation::Square:
-			return LOCTEXT("Square","Final score = Weight * (NormalizedItemValue * NormalizedItemValue)\nBias towards items with big values.").ToString();
+			return LOCTEXT("Square","Final score = ScoringFactor * (NormalizedItemValue * NormalizedItemValue)\nBias towards items with big values.").ToString();
 
 		case EEnvTestScoreEquation::InverseLinear:
-			return LOCTEXT("Inverse","Final score = Weight * (1.0 - NormalizedItemValue)\nBias towards items with values close to zero.  (Linear, but flipped from 1 to 0 rather than 0 to 1.").ToString();
+			return LOCTEXT("Inverse","Final score = ScoringFactor * (1.0 - NormalizedItemValue)\nBias towards items with values close to zero.  (Linear, but flipped from 1 to 0 rather than 0 to 1.").ToString();
 
 		case EEnvTestScoreEquation::Constant:
-			return LOCTEXT("Constant", "Final score (for values that 'pass') = Weight\nNOTE: In this case, the score is normally EITHER the Weight value or zero.\nThe score will be zero if the Normalized Test Value is zero (or if the test value is false for a boolean query).\nOtherwise, score will be the Weight.").ToString();
+			return LOCTEXT("Constant", "Final score (for values that 'pass') = ScoringFactor\nNOTE: In this case, the score is normally EITHER the ScoringFactor value or zero.\nThe score will be zero if the Normalized Test Value is zero (or if the test value is false for a boolean query).\nOtherwise, score will be the ScoringFactor.").ToString();
 
 		default:
 			break;
@@ -581,7 +496,7 @@ FString FEnvQueryTestDetails::GetScoreEquationInfo() const
 	return TEXT("");
 }
 
-EVisibility FEnvQueryTestDetails::GetVisibilityOfFilterMinForScoreClamping() const
+EVisibility FEnvQueryTestDetails::GetVisibilityOfValueMinForScoreClamping() const
 {
 	if (GetFloatScoreVisibility() == EVisibility::Visible)
 	{
@@ -595,7 +510,7 @@ EVisibility FEnvQueryTestDetails::GetVisibilityOfFilterMinForScoreClamping() con
 	return EVisibility::Collapsed;
 }
 
-EVisibility FEnvQueryTestDetails::GetVisibilityOfFilterMaxForScoreClamping() const
+EVisibility FEnvQueryTestDetails::GetVisibilityOfValueMaxForScoreClamping() const
 {
 	if (GetFloatScoreVisibility() == EVisibility::Visible)
 	{
@@ -609,7 +524,7 @@ EVisibility FEnvQueryTestDetails::GetVisibilityOfFilterMaxForScoreClamping() con
 	return EVisibility::Collapsed;
 }
 
-EVisibility FEnvQueryTestDetails::GetVisibilityOfScoreClampingMinimum() const
+EVisibility FEnvQueryTestDetails::GetVisibilityOfScoreClampMinimum() const
 {
 	if (GetFloatScoreVisibility() == EVisibility::Visible)
 	{
@@ -623,7 +538,7 @@ EVisibility FEnvQueryTestDetails::GetVisibilityOfScoreClampingMinimum() const
 	return EVisibility::Collapsed;
 }
 
-EVisibility FEnvQueryTestDetails::GetVisibilityOfScoreClampingMaximum() const
+EVisibility FEnvQueryTestDetails::GetVisibilityOfScoreClampMaximum() const
 {
 	if (GetFloatScoreVisibility() == EVisibility::Visible)
 	{
@@ -694,6 +609,20 @@ EVisibility FEnvQueryTestDetails::GetScoreVisibility() const
 	return EVisibility::Collapsed;
 }
 
+EVisibility FEnvQueryTestDetails::GetClampingVisibility() const
+{
+	if (IsScoring())
+	{
+		const UEnvQueryTest* MyTestOb = Cast<const UEnvQueryTest>(MyTest.Get());
+		if (MyTestOb && MyTestOb->GetWorkOnFloatValues())
+		{
+			return EVisibility::Visible;
+		}
+	}
+
+	return EVisibility::Collapsed;
+}
+
 EVisibility FEnvQueryTestDetails::GetFloatScoreVisibility() const
 {
 	if (IsScoring())
@@ -726,7 +655,7 @@ bool FEnvQueryTestDetails::UsesFilterMax() const
 	return ((EnumValue == EEnvTestFilterType::Maximum) || (EnumValue == EEnvTestFilterType::Range));
 }
 
-EVisibility FEnvQueryTestDetails::GetVisibilityOfFloatFilterMin() const
+EVisibility FEnvQueryTestDetails::GetVisibilityOfFloatValueMin() const
 {
 	if (IsFiltering())
 	{
@@ -744,7 +673,7 @@ EVisibility FEnvQueryTestDetails::GetVisibilityOfFloatFilterMin() const
 	return EVisibility::Collapsed;
 }
 
-EVisibility FEnvQueryTestDetails::GetVisibilityOfFloatFilterMax() const
+EVisibility FEnvQueryTestDetails::GetVisibilityOfFloatValueMax() const
 {
 	if (IsFiltering())
 	{
@@ -783,7 +712,7 @@ bool FEnvQueryTestDetails::IsMatchingBoolValue() const
 	return false;
 }
 
-EVisibility FEnvQueryTestDetails::GetBoolFilterVisibilityForScoring() const
+EVisibility FEnvQueryTestDetails::GetBoolValueVisibilityForScoring() const
 {
 	if (!IsFiltering())
 	{
@@ -796,7 +725,7 @@ EVisibility FEnvQueryTestDetails::GetBoolFilterVisibilityForScoring() const
 	return EVisibility::Collapsed;
 }
 
-EVisibility FEnvQueryTestDetails::GetBoolFilterVisibility() const
+EVisibility FEnvQueryTestDetails::GetBoolValueVisibility() const
 {
 	if (IsFiltering())
 	{
@@ -809,14 +738,6 @@ EVisibility FEnvQueryTestDetails::GetBoolFilterVisibility() const
 	return EVisibility::Collapsed;
 }
 
-EVisibility FEnvQueryTestDetails::GetDiscardFailedVisibility() const
-{
-	uint8 EnumValue;
-	ConditionHandle->GetValue(EnumValue);
-
-	return (EnumValue == EEnvTestCondition::NoCondition) ? EVisibility::Collapsed : EVisibility::Visible;
-}
-
 EVisibility FEnvQueryTestDetails::GetTestPreviewVisibility() const
 {
 	const UEnvQueryTest* MyTestOb = Cast<const UEnvQueryTest>(MyTest.Get());
diff --git a/Engine/Source/Editor/EnvironmentQueryEditor/Private/DetailCustomizations/EnvQueryTestDetails.h b/Engine/Source/Editor/EnvironmentQueryEditor/Private/DetailCustomizations/EnvQueryTestDetails.h
--- a/Engine/Source/Editor/EnvironmentQueryEditor/Private/DetailCustomizations/EnvQueryTestDetails.h
+++ b/Engine/Source/Editor/EnvironmentQueryEditor/Private/DetailCustomizations/EnvQueryTestDetails.h
@@ -27,16 +27,15 @@ protected:
 		FStringIntPair(FString InStr, int32 InInt) : Str(InStr), Int(InInt) {}
 	};
 
-	TSharedPtr<IPropertyHandle> ConditionHandle;
 	TSharedPtr<IPropertyHandle> FilterTypeHandle;
 	TSharedPtr<IPropertyHandle> ScoreEquationHandle;
 	TSharedPtr<IPropertyHandle> TestPurposeHandle;
 	TSharedPtr<IPropertyHandle> ClampMinTypeHandle;
 	TSharedPtr<IPropertyHandle> ClampMaxTypeHandle;
-	TSharedPtr<IPropertyHandle> ScoreClampingMinHandle;
-	TSharedPtr<IPropertyHandle> FloatFilterMinHandle;
-	TSharedPtr<IPropertyHandle> ScoreClampingMaxHandle;
-	TSharedPtr<IPropertyHandle> FloatFilterMaxHandle;
+	TSharedPtr<IPropertyHandle> ScoreClampMinHandle;
+	TSharedPtr<IPropertyHandle> FloatValueMinHandle;
+	TSharedPtr<IPropertyHandle> ScoreClampMaxHandle;
+	TSharedPtr<IPropertyHandle> FloatValueMaxHandle;
 	TSharedPtr<IPropertyHandle> ScoreHandle;
 
 	bool IsFiltering() const;
@@ -45,14 +44,11 @@ protected:
 	bool UsesFilterMin() const;
 	bool UsesFilterMax() const;
 
-	void BuildConditionValues();
-	void OnConditionComboChange(int32 Index);
-	TSharedRef<SWidget> OnGetConditionContent();
-	FString GetCurrentConditionDesc() const;
 	FString GetCurrentFilterTestDesc() const;
 	FString GetScoreEquationInfo() const;
 
  	EVisibility GetScoreVisibility() const;
+	EVisibility GetClampingVisibility() const;
 
 	TSharedRef<SWidget> OnGetFilterTestContent();
 	void BuildFilterTestValues();
@@ -84,15 +80,14 @@ protected:
 	EVisibility GetFloatScoreVisibility() const;
 	
 	EVisibility GetTestPreviewVisibility() const;
-	EVisibility GetVisibilityOfFloatFilterMin() const;
-	EVisibility GetVisibilityOfFloatFilterMax() const;
-	EVisibility GetVisibilityOfFilterMinForScoreClamping() const;
-	EVisibility GetVisibilityOfFilterMaxForScoreClamping() const;
-	EVisibility GetBoolFilterVisibilityForScoring() const;
-	EVisibility GetBoolFilterVisibility() const;
-	EVisibility GetDiscardFailedVisibility() const;
-	EVisibility GetVisibilityOfScoreClampingMinimum() const;
-	EVisibility GetVisibilityOfScoreClampingMaximum() const;
+	EVisibility GetVisibilityOfFloatValueMin() const;
+	EVisibility GetVisibilityOfFloatValueMax() const;
+	EVisibility GetVisibilityOfValueMinForScoreClamping() const;
+	EVisibility GetVisibilityOfValueMaxForScoreClamping() const;
+	EVisibility GetBoolValueVisibilityForScoring() const;
+	EVisibility GetBoolValueVisibility() const;
+	EVisibility GetVisibilityOfScoreClampMinimum() const;
+	EVisibility GetVisibilityOfScoreClampMaximum() const;
 
 	void BuildScoreClampingTypeValues(bool bBuildMinValues, TArray<FStringIntPair>& ClampTypeValues) const;
 	
@@ -100,7 +95,6 @@ protected:
 	void FillEquationSamples(uint8 EquationType, bool bInversed, TArray<float>& Samples) const;
 	TSharedPtr<STestFunctionWidget> PreviewWidget;
 
-	TArray<FStringIntPair> ConditionValues;
 
 	TArray<FStringIntPair> FilterTestValues;
 
diff --git a/Engine/Source/Editor/EnvironmentQueryEditor/Private/EnvironmentQueryGraphNode_Option.cpp b/Engine/Source/Editor/EnvironmentQueryEditor/Private/EnvironmentQueryGraphNode_Option.cpp
--- a/Engine/Source/Editor/EnvironmentQueryEditor/Private/EnvironmentQueryGraphNode_Option.cpp
+++ b/Engine/Source/Editor/EnvironmentQueryEditor/Private/EnvironmentQueryGraphNode_Option.cpp
@@ -25,6 +25,7 @@ void UEnvironmentQueryGraphNode_Option::PostPlacedNewNode()
 		UEnvQuery* Query = Cast<UEnvQuery>(GetEnvironmentQueryGraph()->GetOuter());
 		UEnvQueryOption* QueryOption = ConstructObject<UEnvQueryOption>(UEnvQueryOption::StaticClass(), Query);
 		QueryOption->Generator = ConstructObject<UEnvQueryGenerator>(EnvQueryNodeClass, Query);
+		QueryOption->Generator->UpdateGeneratorVersion();
 		
 		QueryOption->SetFlags(RF_Transactional);
 		QueryOption->Generator->SetFlags(RF_Transactional);
diff --git a/Engine/Source/Runtime/AIModule/Classes/DataProviders/AIDataProvider.h b/Engine/Source/Runtime/AIModule/Classes/DataProviders/AIDataProvider.h
new file mode 100644
--- /dev/null
+++ b/Engine/Source/Runtime/AIModule/Classes/DataProviders/AIDataProvider.h
@@ -0,0 +1,127 @@
+// Copyright 1998-2015 Epic Games, Inc. All Rights Reserved.
+
+#pragma once
+
+#include "AIDataProvider.generated.h"
+
+class APawn;
+class UAIDataProvider;
+
+/**
+ * AIDataProvider is an object that can provide collection of properties
+ * associated with bound pawn owner or request Id.
+ *
+ * Editable properties are used to set up provider instance,
+ * creating additional filters or ways of accessing data (e.g. gameplay tag of ability)
+ *
+ * Non editable properties are holding data
+ */
+
+USTRUCT()
+struct AIMODULE_API FAIDataProviderValue
+{
+	GENERATED_USTRUCT_BODY()
+
+	/** cached uproperty of provider */
+	UPROPERTY(transient)
+	mutable UProperty* CachedProperty;
+
+	/** (optional) provider for dynamic data binding */
+	UPROPERTY(EditAnywhere, Instanced, Category = Value)
+	UAIDataProvider* DataBinding;
+
+	/** name of provider's value property */
+	UPROPERTY(EditAnywhere, Category = Value)
+	FName DataField;
+
+	/** describe default data */
+	virtual FString ValueToString() const;
+	FString ToString() const;
+
+	/** filter for provider's properties */
+	virtual bool IsMatchingType(UProperty* PropType) const;
+
+	/** find all properties of provider that are matching filter */
+	void GetMatchingProperties(TArray<FName>& MatchingProperties) const;
+
+	/** return raw data from provider's property */
+	template<typename T>
+	T* GetRawValue() const;
+
+	/** bind data in provider and cache property for faster access */
+	void BindData(UObject* Owner, int32 RequestId) const;
+
+	FORCEINLINE bool IsDynamic() const { return DataBinding != nullptr; }
+};
+
+USTRUCT()
+struct AIMODULE_API FAIDataProviderTypedValue : public FAIDataProviderValue
+{
+	GENERATED_USTRUCT_BODY()
+
+	/** type of value */
+	UPROPERTY()
+	TSubclassOf<UProperty> PropertyType;
+
+	/** filter for provider's properties */
+	virtual bool IsMatchingType(UProperty* PropType) const override;
+};
+
+USTRUCT()
+struct AIMODULE_API FAIDataProviderStructValue : public FAIDataProviderValue
+{
+	GENERATED_USTRUCT_BODY()
+
+	/** name of UStruct type */
+	FString StructName;
+
+	virtual bool IsMatchingType(UProperty* PropType) const override;
+};
+
+USTRUCT()
+struct AIMODULE_API FAIDataProviderIntValue : public FAIDataProviderTypedValue
+{
+	GENERATED_USTRUCT_BODY()
+	FAIDataProviderIntValue();
+
+	UPROPERTY(EditAnywhere, Category = Value)
+	int32 DefaultValue;
+
+	int32 GetValue() const;
+	virtual FString ValueToString() const override;
+};
+
+USTRUCT()
+struct AIMODULE_API FAIDataProviderFloatValue : public FAIDataProviderTypedValue
+{
+	GENERATED_USTRUCT_BODY()
+	FAIDataProviderFloatValue();
+
+	UPROPERTY(EditAnywhere, Category = Value)
+	float DefaultValue;
+
+	float GetValue() const;
+	virtual FString ValueToString() const override;
+};
+
+USTRUCT()
+struct AIMODULE_API FAIDataProviderBoolValue : public FAIDataProviderTypedValue
+{
+	GENERATED_USTRUCT_BODY()
+	FAIDataProviderBoolValue();
+
+	UPROPERTY(EditAnywhere, Category = Value)
+	bool DefaultValue;
+
+	bool GetValue() const;
+	virtual FString ValueToString() const override;
+};
+
+UCLASS(EditInlineNew, Abstract, CollapseCategories, AutoExpandCategories=(Provider))
+class AIMODULE_API UAIDataProvider : public UObject
+{
+	GENERATED_UCLASS_BODY()
+
+	virtual void BindData(UObject* Owner, int32 RequestId);
+	virtual FString ToString(FName PropName) const;
+};
diff --git a/Engine/Source/Runtime/AIModule/Classes/DataProviders/AIDataProvider_QueryParams.h b/Engine/Source/Runtime/AIModule/Classes/DataProviders/AIDataProvider_QueryParams.h
new file mode 100644
--- /dev/null
+++ b/Engine/Source/Runtime/AIModule/Classes/DataProviders/AIDataProvider_QueryParams.h
@@ -0,0 +1,35 @@
+// Copyright 1998-2015 Epic Games, Inc. All Rights Reserved.
+
+#pragma once
+
+#include "AIDataProvider.h"
+#include "AIDataProvider_QueryParams.generated.h"
+
+/**
+ * AIDataProvider_QueryParams is used with environment queries
+ *
+ * It allows defining simple parameters for running query,
+ * which are not tied to any specific pawn, but defined
+ * for every query execution.
+ */
+
+UCLASS(EditInlineNew, meta=(DisplayName="Query Params"))
+class AIMODULE_API UAIDataProvider_QueryParams : public UAIDataProvider
+{
+	GENERATED_UCLASS_BODY()
+
+	virtual void BindData(UObject* Owner, int32 RequestId);
+	virtual FString ToString(FName PropName) const override;
+
+	UPROPERTY(EditAnywhere, Category = Provider)
+	FName ParamName;
+
+	UPROPERTY()
+	float FloatValue;
+
+	UPROPERTY()
+	int32 IntValue;
+
+	UPROPERTY()
+	bool BoolValue;
+};
diff --git a/Engine/Source/Runtime/AIModule/Classes/EnvironmentQuery/EnvQueryGenerator.h b/Engine/Source/Runtime/AIModule/Classes/EnvironmentQuery/EnvQueryGenerator.h
--- a/Engine/Source/Runtime/AIModule/Classes/EnvironmentQuery/EnvQueryGenerator.h
+++ b/Engine/Source/Runtime/AIModule/Classes/EnvironmentQuery/EnvQueryGenerator.h
@@ -2,10 +2,19 @@
 
 #pragma once
 #include "EnvironmentQuery/EnvQueryTypes.h"
+#include "DataProviders/AIDataProvider.h"
 #include "EnvQueryGenerator.generated.h"
 
 class UEnvQueryItemType;
 
+namespace EnvQueryGeneratorVersion
+{
+	const int32 Initial = 0;
+	const int32 DataProviders = 1;
+
+	const int32 Latest = DataProviders;
+}
+
 UCLASS(Abstract)
 class AIMODULE_API UEnvQueryGenerator : public UObject
 {
@@ -18,6 +27,10 @@ class AIMODULE_API UEnvQueryGenerator : public UObject
 	UPROPERTY()
 	TSubclassOf<UEnvQueryItemType> ItemType;
 
+	/** versioning for updating deprecated properties */
+	UPROPERTY()
+	int32 VerNum;
+
 	virtual void GenerateItems(FEnvQueryInstance& QueryInstance) const { checkNoEntry(); }
 
 	/** get description of generator */
@@ -27,4 +40,7 @@ class AIMODULE_API UEnvQueryGenerator : public UObject
 #if WITH_EDITOR && USE_EQS_DEBUGGER
 	virtual void PostEditChangeProperty(struct FPropertyChangedEvent& PropertyChangedEvent) override;
 #endif //WITH_EDITOR && USE_EQS_DEBUGGER
+
+	virtual void PostLoad() override;
+	void UpdateGeneratorVersion();
 };
diff --git a/Engine/Source/Runtime/AIModule/Classes/EnvironmentQuery/EnvQueryManager.h b/Engine/Source/Runtime/AIModule/Classes/EnvironmentQuery/EnvQueryManager.h
--- a/Engine/Source/Runtime/AIModule/Classes/EnvironmentQuery/EnvQueryManager.h
+++ b/Engine/Source/Runtime/AIModule/Classes/EnvironmentQuery/EnvQueryManager.h
@@ -141,6 +141,9 @@ class AIMODULE_API UEnvQueryManager : public UObject, public FTickableGameObject
 	/** creates local context object */
 	UEnvQueryContext* PrepareLocalContext(TSubclassOf<UEnvQueryContext> ContextClass);
 
+	/** find value of named param stored with active query */
+	float FindNamedParam(int32 QueryId, FName ParamName) const;
+
 	/** execute query */
 	bool AbortQuery(int32 RequestID);
 
diff --git a/Engine/Source/Runtime/AIModule/Classes/EnvironmentQuery/EnvQueryTest.h b/Engine/Source/Runtime/AIModule/Classes/EnvironmentQuery/EnvQueryTest.h
--- a/Engine/Source/Runtime/AIModule/Classes/EnvironmentQuery/EnvQueryTest.h
+++ b/Engine/Source/Runtime/AIModule/Classes/EnvironmentQuery/EnvQueryTest.h
@@ -1,6 +1,7 @@
 // Copyright 1998-2015 Epic Games, Inc. All Rights Reserved.
 
 #pragma once
+#include "DataProviders/AIDataProvider.h"
 #include "EnvironmentQuery/EnvQueryTypes.h"
 #include "EnvQueryTest.generated.h"
 
@@ -11,6 +12,14 @@ class UEnvQueryContext;
 struct FPropertyChangedEvent;
 #endif // WITH_EDITOR
 
+namespace EnvQueryTestVersion
+{
+	const int32 Initial = 0;
+	const int32 DataProviders = 1;
+
+	const int32 Latest = DataProviders;
+}
+
 UCLASS(Abstract)
 class AIMODULE_API UEnvQueryTest : public UObject
 {
@@ -20,75 +29,58 @@ class AIMODULE_API UEnvQueryTest : public UObject
 	UPROPERTY()
 	int32 TestOrder;
 
-	UPROPERTY(EditDefaultsOnly, Category=Test,
-		Meta=(ToolTip="The purpose of this test.  Should it be used for filtering possible results, scoring them, or both?"))
+	/** Versioning for updating deprecated properties */
+	UPROPERTY()
+	int32 VerNum;
+
+	/** The purpose of this test.  Should it be used for filtering possible results, scoring them, or both? */
+	UPROPERTY(EditDefaultsOnly, Category=Test)
 	TEnumAsByte<EEnvTestPurpose::Type> TestPurpose;
 	
-	UPROPERTY(EditDefaultsOnly, Category=Filter,
-		Meta=(ToolTip="Does this test filter out results that are below a lower limit, above an upper limit, or both?  Or does it just look for a matching value?"))
+	/** Does this test filter out results that are below a lower limit, above an upper limit, or both?  Or does it just look for a matching value? */
+	UPROPERTY(EditDefaultsOnly, Category=Filter)
 	TEnumAsByte<EEnvTestFilterType::Type> FilterType;
 
 	/** Boolean value that must be matched in order for scoring to occur or filtering test to pass. */
-	UPROPERTY(EditDefaultsOnly, Category=Filter,
-		Meta=(ToolTip="Boolean value that must be matched in order for scoring to occur or filtering test to pass."))
-	FEnvBoolParam BoolFilter;
+	UPROPERTY(EditDefaultsOnly, Category=Filter)
+	FAIDataProviderBoolValue BoolValue;
 
-	// TODO: REMOVE, Deprecated!
-	/** filter for test value; currently this is the ONLY value, because we can only filter in one direction. */
-	UPROPERTY()
-	FEnvFloatParam FloatFilter;
+	/** Minimum limit (inclusive) of valid values for the raw test value. Lower values will be discarded as invalid. */
+	UPROPERTY(EditDefaultsOnly, Category=Filter)
+	FAIDataProviderFloatValue FloatValueMin;
 
-	UPROPERTY(EditDefaultsOnly, Category=Filter,
-		Meta=(ToolTip="Minimum limit (inclusive) of valid values for the raw test value.  Lower values will be discarded as invalid."))
-	FEnvFloatParam FloatFilterMin;
-
-	UPROPERTY(EditDefaultsOnly, Category=Filter,
-		Meta=(ToolTip="Maximum limit (inclusive) of valid values for the raw test value.  Higher values will be discarded as invalid."))
-	FEnvFloatParam FloatFilterMax;
-
-	// Using this property just as a way to do file version conversion without having to up the version number!
-	UPROPERTY()
-	bool bFormatUpdated;
-
-	// TODO: REMOVE, Deprecated!
-	/** Condition for discarding item */
-	UPROPERTY()
-	TEnumAsByte<EEnvTestCondition::Type> Condition;
-
-	// TODO: REMOVE, Deprecated!
-	/** Weight modifier */
-	UPROPERTY()
-	TEnumAsByte<EEnvTestWeight::Type> WeightModifier;
+	/** Maximum limit (inclusive) of valid values for the raw test value. Higher values will be discarded as invalid. */
+	UPROPERTY(EditDefaultsOnly, Category=Filter)
+	FAIDataProviderFloatValue FloatValueMax;
 
 	/** Cost of test */
 	TEnumAsByte<EEnvTestCost::Type> Cost;
 
-	UPROPERTY(EditDefaultsOnly, Category=Score,
-		Meta=(ToolTip="How should the lower bound for normalization of the raw test value before applying the scoring formula be determined?  Should it use the lowest value found (tested), the lower threshold for filtering, or a separate specified normalization minimum?"))
+	/** How should the lower bound for normalization of the raw test value before applying the scoring formula be determined?
+	    Should it use the lowest value found (tested), the lower threshold for filtering, or a separate specified normalization minimum? */
+	UPROPERTY(EditDefaultsOnly, Category=Score)
 	TEnumAsByte<EEnvQueryTestClamping::Type> ClampMinType;
 
-	UPROPERTY(EditDefaultsOnly, Category=Score,
-		Meta=(ToolTip="How should the upper bound for normalization of the raw test value before applying the scoring formula be determined?  Should it use the highes value found (tested), the upper threshold for filtering, or a separate specified normalization maximum?"))
+	/** How should the upper bound for normalization of the raw test value before applying the scoring formula be determined?
+	    Should it use the highest value found (tested), the upper threshold for filtering, or a separate specified normalization maximum? */
+	UPROPERTY(EditDefaultsOnly, Category=Score)
 	TEnumAsByte<EEnvQueryTestClamping::Type> ClampMaxType;
 
-	UPROPERTY(EditDefaultsOnly, Category=Score,
-		Meta=(ToolTip="Minimum value to use to normalize the raw test value before applying scoring formula."))
-	FEnvFloatParam ScoreClampingMin;
+	/** Minimum value to use to normalize the raw test value before applying scoring formula. */
+	UPROPERTY(EditDefaultsOnly, Category=Score)
+	FAIDataProviderFloatValue ScoreClampMin;
 
-	UPROPERTY(EditDefaultsOnly, Category=Score,
-		Meta=(ToolTip="Maximum value to use to normalize the raw test value before applying scoring formula."))
-	FEnvFloatParam ScoreClampingMax;
+	/** Maximum value to use to normalize the raw test value before applying scoring formula. */
+	UPROPERTY(EditDefaultsOnly, Category=Score)
+	FAIDataProviderFloatValue ScoreClampMax;
 
-	UPROPERTY(EditDefaultsOnly, Category=Score,
-		Meta=(ToolTip="The shape of the curve equation to apply to the normalized score before multiplying by Weight."))
+	/** The shape of the curve equation to apply to the normalized score before multiplying by factor. */
+	UPROPERTY(EditDefaultsOnly, Category=Score)
 	TEnumAsByte<EEnvTestScoreEquation::Type> ScoringEquation;
 
-	// For now, I'm keeping "Weight".  TODO: Reconsider Weight vs. "Scoring Factor", Constant, Power, etc. once we have a final plan for what properties are available.
-	/** Weight of test */
-	UPROPERTY(EditDefaultsOnly, Category=Score,
-		Meta=(ToolTip="The weight (factor) by which to multiply the normalized score after the scoring equation is applied.",
-			  ClampMin="0.001", UIMin="0.001"))
-	FEnvFloatParam Weight;
+	/** The weight (factor) by which to multiply the normalized score after the scoring equation is applied. */
+	UPROPERTY(EditDefaultsOnly, Category=Score, Meta=(ClampMin="0.001", UIMin="0.001"))
+	FAIDataProviderFloatValue ScoringFactor;
 
 	/** Validation: item type that can be used with this test */
 	TSubclassOf<UEnvQueryItemType> ValidItemType;
@@ -103,9 +95,25 @@ class AIMODULE_API UEnvQueryTest : public UObject
 					|| (ScoringEquation == EEnvTestScoreEquation::Constant));	// We are giving a constant score value for passing.
 	}
 
-	// TODO: REMOVE, Deprecated!
-	/** Normalize test result starting from 0 */
-	uint32 bNormalizeFromZero : 1;
+	// BEGIN: deprecated properties, do not use them
+	UPROPERTY()
+	FEnvBoolParam BoolFilter;
+
+	UPROPERTY()
+	FEnvFloatParam FloatFilterMin;
+
+	UPROPERTY()
+	FEnvFloatParam FloatFilterMax;
+
+	UPROPERTY()
+	FEnvFloatParam ScoreClampingMin;
+
+	UPROPERTY()
+	FEnvFloatParam ScoreClampingMax;
+
+	UPROPERTY()
+	FEnvFloatParam Weight;
+	// END: deprecated properties, do not use them
 
 private:
 	/** When set, test operates on float values (e.g. distance, with AtLeast, UpTo conditions),
@@ -114,10 +122,6 @@ private:
 	uint32 bWorkOnFloatValues : 1;
 
 public:
-	// TODO: REMOVE, Deprecated!
-	/** If set, items not passing filter condition will be removed from query, otherwise they will receive score 0 */
-	UPROPERTY()
-	uint32 bDiscardFailedItems : 1;
 
 	/** Function that does the actual work */
 	virtual void RunTest(FEnvQueryInstance& QueryInstance) const { checkNoEntry(); }
@@ -140,8 +144,8 @@ public:
 	/** normalize scores in range */
 	void NormalizeItemScores(FEnvQueryInstance& QueryInstance);
 
-	FORCEINLINE bool IsScoring() const { return (TestPurpose != EEnvTestPurpose::Filter); } // ((TestPurpose == EEnvTestPurpose::Score) || (TestPurpose == EEnvTestPurpose::FilterAndScore));
-	FORCEINLINE bool IsFiltering() const { return (TestPurpose != EEnvTestPurpose::Score); } // ((TestPurpose == EEnvTestPurpose::Filter) || (TestPurpose == EEnvTestPurpose::FilterAndScore));
+	FORCEINLINE bool IsScoring() const { return (TestPurpose != EEnvTestPurpose::Filter); } 
+	FORCEINLINE bool IsFiltering() const { return (TestPurpose != EEnvTestPurpose::Score); }
 
 	/** get description of test */
 	virtual FString GetDescriptionTitle() const;
diff --git a/Engine/Source/Runtime/AIModule/Classes/EnvironmentQuery/EnvQueryTypes.h b/Engine/Source/Runtime/AIModule/Classes/EnvironmentQuery/EnvQueryTypes.h
--- a/Engine/Source/Runtime/AIModule/Classes/EnvironmentQuery/EnvQueryTypes.h
+++ b/Engine/Source/Runtime/AIModule/Classes/EnvironmentQuery/EnvQueryTypes.h
@@ -3,6 +3,7 @@
 #pragma once
 
 #include "EnvironmentQuery/Items/EnvQueryItemType.h"
+#include "DataProviders/AIDataProvider.h"
 #include "EnvQueryTypes.generated.h"
 
 class ARecastNavMesh;
@@ -56,18 +57,6 @@ namespace EEnvTestFilterType
 	};
 }
 
-UENUM()
-namespace EEnvTestCondition
-{
-	enum Type
-	{
-		NoCondition		UMETA(DisplayName="Always pass"),
-		AtLeast			UMETA(DisplayName="At least"),
-		UpTo			UMETA(DisplayName="Up to"),
-		Match,
-	};
-}
-
 UENUM()
 namespace EEnvTestScoreEquation
 {
@@ -99,7 +88,7 @@ namespace EEnvTestWeight
 		None,
 		Square,
 		Inverse,
-		Absolute,	// Removed after discussion with EGP folks
+		Unused			UMETA(Hidden),
 		Constant,
 		Skip			UMETA(DisplayName = "Do not weight"),
 	};
@@ -195,18 +184,12 @@ namespace EEnvQueryTestClamping
 	};
 }
 
+// DEPRECATED, will be removed soon - use AI Data Providers instead 
 USTRUCT()
 struct AIMODULE_API FEnvFloatParam
 {
 	GENERATED_USTRUCT_BODY();
 
-	typedef float FValueType;
-
-	FEnvFloatParam() :
-		Value(0.f)
-	{
-	}
-
 	/** default value */
 	UPROPERTY(EditDefaultsOnly, Category=Param)
 	float Value;
@@ -216,20 +199,15 @@ struct AIMODULE_API FEnvFloatParam
 	FName ParamName;
 
 	bool IsNamedParam() const { return ParamName != NAME_None; }
+	void Convert(UObject* Owner, FAIDataProviderFloatValue& ValueProvider);
 };
 
+// DEPRECATED, will be removed soon - use AI Data Providers instead 
 USTRUCT()
 struct AIMODULE_API FEnvIntParam
 {
 	GENERATED_USTRUCT_BODY();
 
-	typedef int32 FValueType;
-
-	FEnvIntParam() :
-		Value(0)
-	{
-	}
-
 	/** default value */
 	UPROPERTY(EditDefaultsOnly, Category=Param)
 	int32 Value;
@@ -239,20 +217,15 @@ struct AIMODULE_API FEnvIntParam
 	FName ParamName;
 
 	bool IsNamedParam() const { return ParamName != NAME_None; }
+	void Convert(UObject* Owner, FAIDataProviderIntValue& ValueProvider);
 };
 
+// DEPRECATED, will be removed soon - use AI Data Providers instead 
 USTRUCT()
 struct AIMODULE_API FEnvBoolParam
 {
 	GENERATED_USTRUCT_BODY();
 
-	typedef bool FValueType;
-
-	FEnvBoolParam() :
-		Value(false)
-	{
-	}
-
 	/** default value */
 	UPROPERTY(EditDefaultsOnly, Category=Param)
 	bool Value;
@@ -262,6 +235,7 @@ struct AIMODULE_API FEnvBoolParam
 	FName ParamName;
 
 	bool IsNamedParam() const { return ParamName != NAME_None; }
+	void Convert(UObject* Owner, FAIDataProviderBoolValue& ValueProvider);
 };
 
 USTRUCT(BlueprintType)
@@ -668,30 +642,6 @@ public:
 	/** check if current test can batch its calculations */
 	bool CanBatchTest() const { return !IsInSingleItemFinalSearch(); }
 
-	/** access named params */
-	template<typename TEQSParam>
-	FORCEINLINE bool GetParamValue(const TEQSParam& Param, typename TEQSParam::FValueType& OutValue, const FString& ParamDesc)
-	{
-		if (Param.ParamName != NAME_None)
-		{
-			FNamedParamValueType* PtrValue = NamedParams.Find(Param.ParamName);
-			if (PtrValue == NULL)
-			{
-				EQSHEADERLOG(FString::Printf(TEXT("Query [%s] is missing param [%s] for [%s] property!"), *QueryName, *Param.ParamName.ToString(), *ParamDesc));				
-				MarkAsMissingParam();
-				return false;
-			}
-
-			OutValue = *((typename TEQSParam::FValueType*)PtrValue);
-		}
-		else
-		{
-			OutValue = Param.Value;
-		}
-
-		return true;
-	}
-
 	/** raw data operations */
 	void ReserveItemData(int32 NumAdditionalItems);
 
@@ -828,12 +778,14 @@ public:
 			}
 
 			if (bPassedTest)
-			{	// If we passed the test, either we really did, or we're only scoring, so we can't truly "fail".	
+			{
+				// If we passed the test, either we really did, or we're only scoring, so we can't truly "fail".	
 				ItemScore += Score;
 				NumPartialScores++;
 			}
 			else
-			{	// We are ONLY filtering, and we failed
+			{
+				// We are ONLY filtering, and we failed
 				bPassed = false;
 			}
 		}
@@ -931,14 +883,12 @@ public:
 	protected:
 
 		FEnvQueryInstance* Instance;
-		//UEnvQueryTest* Test;
 		int32 CurrentItem;
 		int32 NumPartialScores;
 		double Deadline;
 		float ItemScore;
 		uint32 bPassed : 1;
 		uint32 bSkipped : 1;
-		uint32 bDiscardFailed : 1;
 
 		void InitItemScore()
 		{
@@ -987,9 +937,5 @@ public:
 	static float SkippedItemValue;
 
 	static FText GetShortTypeName(const UObject* Ob);
-
 	static FText DescribeContext(TSubclassOf<UEnvQueryContext> ContextClass);
-	static FString DescribeIntParam(const FEnvIntParam& Param);
-	static FString DescribeFloatParam(const FEnvFloatParam& Param);
-	static FString DescribeBoolParam(const FEnvBoolParam& Param);
 };
diff --git a/Engine/Source/Runtime/AIModule/Classes/EnvironmentQuery/Generators/EnvQueryGenerator_ActorsOfClass.h b/Engine/Source/Runtime/AIModule/Classes/EnvironmentQuery/Generators/EnvQueryGenerator_ActorsOfClass.h
--- a/Engine/Source/Runtime/AIModule/Classes/EnvironmentQuery/Generators/EnvQueryGenerator_ActorsOfClass.h
+++ b/Engine/Source/Runtime/AIModule/Classes/EnvironmentQuery/Generators/EnvQueryGenerator_ActorsOfClass.h
@@ -14,7 +14,7 @@ class UEnvQueryGenerator_ActorsOfClass : public UEnvQueryGenerator
 
 	/** max distance of path between point and context */
 	UPROPERTY(EditDefaultsOnly, Category=Generator)
-	FEnvFloatParam Radius;
+	FAIDataProviderFloatValue SearchRadius;
 
 	UPROPERTY(EditDefaultsOnly, Category=Generator)
 	TSubclassOf<AActor> SearchedActorClass;
@@ -23,7 +23,12 @@ class UEnvQueryGenerator_ActorsOfClass : public UEnvQueryGenerator
 	UPROPERTY(EditAnywhere, Category=Generator)
 	TSubclassOf<UEnvQueryContext> SearchCenter;
 
+	// DEPRECATED
+	UPROPERTY()
+	FEnvFloatParam Radius;
+
 	virtual void GenerateItems(FEnvQueryInstance& QueryInstance) const override;
+	virtual void PostLoad() override;
 
 	virtual FText GetDescriptionTitle() const override;
 	virtual FText GetDescriptionDetails() const override;
diff --git a/Engine/Source/Runtime/AIModule/Classes/EnvironmentQuery/Generators/EnvQueryGenerator_BlueprintBase.h b/Engine/Source/Runtime/AIModule/Classes/EnvironmentQuery/Generators/EnvQueryGenerator_BlueprintBase.h
--- a/Engine/Source/Runtime/AIModule/Classes/EnvironmentQuery/Generators/EnvQueryGenerator_BlueprintBase.h
+++ b/Engine/Source/Runtime/AIModule/Classes/EnvironmentQuery/Generators/EnvQueryGenerator_BlueprintBase.h
@@ -14,10 +14,6 @@ class AIMODULE_API UEnvQueryGenerator_BlueprintBase : public UEnvQueryGenerator
 {
 	GENERATED_UCLASS_BODY()
 
-	/** max distance of between point and context */
-	UPROPERTY(EditDefaultsOnly, Category = Generator)
-	FEnvFloatParam Radius; 
-
 	/** A short description of what test does, like "Generate pawn named Joe" */
 	UPROPERTY(EditAnywhere, Category = Generator)
 	FText GeneratorsActionDescription;
diff --git a/Engine/Source/Runtime/AIModule/Classes/EnvironmentQuery/Generators/EnvQueryGenerator_OnCircle.h b/Engine/Source/Runtime/AIModule/Classes/EnvironmentQuery/Generators/EnvQueryGenerator_OnCircle.h
--- a/Engine/Source/Runtime/AIModule/Classes/EnvironmentQuery/Generators/EnvQueryGenerator_OnCircle.h
+++ b/Engine/Source/Runtime/AIModule/Classes/EnvironmentQuery/Generators/EnvQueryGenerator_OnCircle.h
@@ -6,28 +6,6 @@
 
 struct FEnvQueryInstance;
 
-class FVectorAndDataContainer
-{
-public:
-	virtual ~FVectorAndDataContainer() {}
-
-	virtual void Reset() { Locations.Empty(); }
-
-	FORCEINLINE int32 GetNumEntries() const { return Locations.Num(); }
-	FORCEINLINE FVector GetVector(int32 Index) const { return Locations[Index]; }
-
-	FORCEINLINE TArray<FVector>& GetLocationArray() { return Locations; }
-
-protected:
-	// TODO: Make inline non-virtual and turn this into the base class?  If so, we wouldn't have to allocate this class
-	// We'd only have to fill in a parallel array if/when needed.
-	
-	// virtual FVectorAdditionalData& GetAdditionalData(int32 Index);
-
-private:
-	TArray<FVector> Locations;
-};
-
 UCLASS()	
 class AIMODULE_API UEnvQueryGenerator_OnCircle : public UEnvQueryGenerator_ProjectedPoints
 {
@@ -35,11 +13,11 @@ class AIMODULE_API UEnvQueryGenerator_OnCircle : public UEnvQueryGenerator_Proje
 
 	/** max distance of path between point and context */
 	UPROPERTY(EditDefaultsOnly, Category=Generator)
-	FEnvFloatParam Radius;
+	FAIDataProviderFloatValue CircleRadius;
 
 	/** items will be generated on a circle this much apart */
-	UPROPERTY(EditDefaultsOnly, Category=Generator)
-	FEnvFloatParam ItemSpacing;
+	UPROPERTY(EditDefaultsOnly, Category = Generator)
+	FAIDataProviderFloatValue SpaceBetween;
 
 	/** If you generate items on a piece of circle you define direction of Arc cut here */
 	UPROPERTY(EditDefaultsOnly, Category=Generator, meta=(EditCondition="bDefineArc"))
@@ -47,7 +25,7 @@ class AIMODULE_API UEnvQueryGenerator_OnCircle : public UEnvQueryGenerator_Proje
 
 	/** If you generate items on a piece of circle you define angle of Arc cut here */
 	UPROPERTY(EditDefaultsOnly, Category=Generator)
-	FEnvFloatParam Angle;
+	FAIDataProviderFloatValue ArcAngle;
 	
 	UPROPERTY()
 	mutable float AngleRadians;
@@ -63,6 +41,17 @@ class AIMODULE_API UEnvQueryGenerator_OnCircle : public UEnvQueryGenerator_Proje
 	UPROPERTY()
 	uint32 bDefineArc:1;
 
+	// BEGIN: deprecated properties
+	UPROPERTY()
+	FEnvFloatParam Radius;
+
+	UPROPERTY()
+	FEnvFloatParam ItemSpacing;
+
+	UPROPERTY()
+	FEnvFloatParam Angle;
+	// END: deprecated properties
+
 	virtual void PostLoad() override;
 
 	virtual void GenerateItems(FEnvQueryInstance& QueryInstance) const override;
@@ -77,20 +66,10 @@ class AIMODULE_API UEnvQueryGenerator_OnCircle : public UEnvQueryGenerator_Proje
 protected:
 	FVector CalcDirection(FEnvQueryInstance& QueryInstance) const;
 
-	// For derived classes that need to store additional data associated with the FVector, override these two functions:
-	virtual void PrepareGeneratorContext(FEnvQueryInstance& QueryInstance,
-										 TSharedPtr<FVectorAndDataContainer>& CenterLocationCandidates) const;
-
-	virtual void AddItemDataForCircle(const TArray<FVector>& ItemCandidates,
-									  TSharedPtr<FVectorAndDataContainer>& CenterLocationCandidates,
-									  int32 CenterLocationCandidateIndex,
-									  FEnvQueryInstance& OutQueryInstance) const;
-
-private:
-	void GenerateItemsForCircle(TSharedPtr<FVectorAndDataContainer>& CenterLocationCandidates,
-								int32 CenterLocationCandidateIndex,
-								const FVector& CenterLocation,
-								const FVector& StartDirection,
-								int32 StepsCount, float AngleStep,
-								FEnvQueryInstance& OutQueryInstance) const;
+	void GenerateItemsForCircle(uint8* ContextRawData, UEnvQueryItemType* ContextItemType,
+		const FVector& CenterLocation, const FVector& StartDirection,
+		int32 StepsCount, float AngleStep, FEnvQueryInstance& OutQueryInstance) const;
+
+	virtual void AddItemDataForCircle(uint8* ContextRawData, UEnvQueryItemType* ContextItemType, 
+		const TArray<FVector>& Locations, FEnvQueryInstance& OutQueryInstance) const;
 };
diff --git a/Engine/Source/Runtime/AIModule/Classes/EnvironmentQuery/Generators/EnvQueryGenerator_PathingGrid.h b/Engine/Source/Runtime/AIModule/Classes/EnvironmentQuery/Generators/EnvQueryGenerator_PathingGrid.h
--- a/Engine/Source/Runtime/AIModule/Classes/EnvironmentQuery/Generators/EnvQueryGenerator_PathingGrid.h
+++ b/Engine/Source/Runtime/AIModule/Classes/EnvironmentQuery/Generators/EnvQueryGenerator_PathingGrid.h
@@ -23,15 +23,15 @@ class UEnvQueryGenerator_PathingGrid : public UEnvQueryGenerator
 
 	/** max distance of path between point and context */
 	UPROPERTY(EditDefaultsOnly, Category=Generator)
-	FEnvFloatParam MaxPathDistance;
+	FAIDataProviderFloatValue MaxDistance;
 
 	/** generation density */
 	UPROPERTY(EditDefaultsOnly, Category=Generator)
-	FEnvFloatParam Density;
+	FAIDataProviderFloatValue SpaceBetween;
 
 	/** path direction switch: set = path from context to points, clear = path from points to context */
 	UPROPERTY(EditDefaultsOnly, Category=Generator)
-	FEnvBoolParam PathFromContext;
+	FAIDataProviderBoolValue PathToItem;
 
 	/** context */
 	UPROPERTY(EditAnywhere, Category=Generator)
@@ -41,7 +41,19 @@ class UEnvQueryGenerator_PathingGrid : public UEnvQueryGenerator
 	UPROPERTY(EditDefaultsOnly, Category=Generator)
 	TSubclassOf<UNavigationQueryFilter> NavigationFilter;
 
+	// BEGIN: deprecated properties
+	UPROPERTY()
+	FEnvFloatParam MaxPathDistance;
+
+	UPROPERTY()
+	FEnvFloatParam Density;
+
+	UPROPERTY()
+	FEnvBoolParam PathFromContext;
+	// END: deprecated properties
+
 	virtual void GenerateItems(FEnvQueryInstance& QueryInstance) const override;
+	virtual void PostLoad() override;
 
 	virtual FText GetDescriptionTitle() const override;
 	virtual FText GetDescriptionDetails() const override;
diff --git a/Engine/Source/Runtime/AIModule/Classes/EnvironmentQuery/Generators/EnvQueryGenerator_SimpleGrid.h b/Engine/Source/Runtime/AIModule/Classes/EnvironmentQuery/Generators/EnvQueryGenerator_SimpleGrid.h
--- a/Engine/Source/Runtime/AIModule/Classes/EnvironmentQuery/Generators/EnvQueryGenerator_SimpleGrid.h
+++ b/Engine/Source/Runtime/AIModule/Classes/EnvironmentQuery/Generators/EnvQueryGenerator_SimpleGrid.h
@@ -17,17 +17,26 @@ class UEnvQueryGenerator_SimpleGrid : public UEnvQueryGenerator_ProjectedPoints
 
 	/** square's extent */
 	UPROPERTY(EditDefaultsOnly, Category=Generator)
-	FEnvFloatParam Radius;
+	FAIDataProviderFloatValue GridSize;
 
 	/** generation density */
 	UPROPERTY(EditDefaultsOnly, Category=Generator)
-	FEnvFloatParam Density;
+	FAIDataProviderFloatValue SpaceBetween;
 
 	/** context */
 	UPROPERTY(EditDefaultsOnly, Category=Generator)
 	TSubclassOf<UEnvQueryContext> GenerateAround;
 
+	// BEGIN: deprecated properties
+	UPROPERTY()
+	FEnvFloatParam Radius;
+
+	UPROPERTY()
+	FEnvFloatParam Density;
+	// END: deprecated properties
+
 	virtual void GenerateItems(FEnvQueryInstance& QueryInstance) const override;
+	virtual void PostLoad() override;
 
 	virtual FText GetDescriptionTitle() const override;
 	virtual FText GetDescriptionDetails() const override;
diff --git a/Engine/Source/Runtime/AIModule/Classes/EnvironmentQuery/Tests/EnvQueryTest_Pathfinding.h b/Engine/Source/Runtime/AIModule/Classes/EnvironmentQuery/Tests/EnvQueryTest_Pathfinding.h
--- a/Engine/Source/Runtime/AIModule/Classes/EnvironmentQuery/Tests/EnvQueryTest_Pathfinding.h
+++ b/Engine/Source/Runtime/AIModule/Classes/EnvironmentQuery/Tests/EnvQueryTest_Pathfinding.h
@@ -32,15 +32,26 @@ class UEnvQueryTest_Pathfinding : public UEnvQueryTest
 
 	/** pathfinding direction */
 	UPROPERTY(EditDefaultsOnly, Category=Pathfinding)
-	FEnvBoolParam PathToItem;
+	FAIDataProviderBoolValue PathFromContext;
 
 	/** if set, items with failed path will be invalidated (PathCost, PathLength) */
 	UPROPERTY(EditDefaultsOnly, Category=Pathfinding, AdvancedDisplay)
-	FEnvBoolParam DiscardUnreachable;
+	FAIDataProviderBoolValue SkipUnreachable;
 
 	/** if set, hierarchical (faster) pathfinding will be used */
 	UPROPERTY(EditDefaultsOnly, Category=Pathfinding, AdvancedDisplay)
+	FAIDataProviderBoolValue UseHierarchicalPathfinding;
+
+	// BEGIN: deprecated properties 
+	UPROPERTY()
+	FEnvBoolParam PathToItem;
+
+	UPROPERTY()
+	FEnvBoolParam DiscardUnreachable;
+
+	UPROPERTY()
 	FEnvBoolParam HierarchicalPathfinding;
+	// END: deprecated properties
 
 	virtual void RunTest(FEnvQueryInstance& QueryInstance) const override;
 
diff --git a/Engine/Source/Runtime/AIModule/Classes/EnvironmentQuery/Tests/EnvQueryTest_Trace.h b/Engine/Source/Runtime/AIModule/Classes/EnvironmentQuery/Tests/EnvQueryTest_Trace.h
--- a/Engine/Source/Runtime/AIModule/Classes/EnvironmentQuery/Tests/EnvQueryTest_Trace.h
+++ b/Engine/Source/Runtime/AIModule/Classes/EnvironmentQuery/Tests/EnvQueryTest_Trace.h
@@ -17,21 +17,33 @@ class UEnvQueryTest_Trace : public UEnvQueryTest
 
 	/** trace direction */
 	UPROPERTY(EditDefaultsOnly, Category=Trace)
-	FEnvBoolParam TraceToItem;
+	FAIDataProviderBoolValue TraceFromContext;
 
 	/** Z offset from item */
 	UPROPERTY(EditDefaultsOnly, Category=Trace, AdvancedDisplay)
-	FEnvFloatParam ItemOffsetZ;
+	FAIDataProviderFloatValue ItemHeightOffset;
 
 	/** Z offset from querier */
 	UPROPERTY(EditDefaultsOnly, Category=Trace, AdvancedDisplay)
-	FEnvFloatParam ContextOffsetZ;
+	FAIDataProviderFloatValue ContextHeightOffset;
 
 	/** context: other end of trace test */
 	UPROPERTY(EditDefaultsOnly, Category=Trace)
 	TSubclassOf<UEnvQueryContext> Context;
 
+	// BEGIN: deprecated properties 
+	UPROPERTY()
+	FEnvBoolParam TraceToItem;
+
+	UPROPERTY()
+	FEnvFloatParam ItemOffsetZ;
+
+	UPROPERTY()
+	FEnvFloatParam ContextOffsetZ;
+	// END: deprecated properties
+
 	virtual void RunTest(FEnvQueryInstance& QueryInstance) const override;
+	virtual void PostLoad() override;
 
 	virtual FString GetDescriptionTitle() const override;
 	virtual FText GetDescriptionDetails() const override;
diff --git a/Engine/Source/Runtime/AIModule/Private/DataProviders/AIDataProvider.cpp b/Engine/Source/Runtime/AIModule/Private/DataProviders/AIDataProvider.cpp
new file mode 100644
--- /dev/null
+++ b/Engine/Source/Runtime/AIModule/Private/DataProviders/AIDataProvider.cpp
@@ -0,0 +1,164 @@
+// Copyright 1998-2015 Epic Games, Inc. All Rights Reserved.
+
+#include "AIModulePrivate.h"
+#include "DataProviders/AIDataProvider.h"
+
+//////////////////////////////////////////////////////////////////////////
+// FAIDataProviderValue
+
+bool FAIDataProviderValue::IsMatchingType(UProperty* PropType) const
+{
+	return true;
+}
+
+void FAIDataProviderValue::GetMatchingProperties(TArray<FName>& MatchingProperties) const
+{
+	if (DataBinding)
+	{
+		for (UProperty* Prop = DataBinding->GetClass()->PropertyLink; Prop; Prop = Prop->PropertyLinkNext)
+		{
+			if (Prop->HasAnyPropertyFlags(CPF_Edit))
+			{
+				continue;
+			}
+
+			if (IsMatchingType(Prop))
+			{
+				MatchingProperties.Add(Prop->GetFName());
+			}
+		}
+	}
+}
+
+template<typename T>
+T* FAIDataProviderValue::GetRawValue() const
+{
+	return CachedProperty ? CachedProperty->ContainerPtrToValuePtr<T>(DataBinding) : nullptr;
+}
+
+void FAIDataProviderValue::BindData(UObject* Owner, int32 RequestId) const
+{
+	if (DataBinding)
+	{
+		DataBinding->BindData(Owner, RequestId);
+		CachedProperty = DataBinding->GetClass()->FindPropertyByName(DataField);
+	}
+}
+
+FString FAIDataProviderValue::ToString() const
+{
+	return IsDynamic() ? DataBinding->ToString(DataField) : ValueToString();
+}
+
+FString FAIDataProviderValue::ValueToString() const
+{
+	return TEXT("??");
+}
+
+//////////////////////////////////////////////////////////////////////////
+// FAIDataProviderTypedValue
+
+bool FAIDataProviderTypedValue::IsMatchingType(UProperty* PropType) const
+{
+	return PropType->GetClass() == PropertyType;
+}
+
+//////////////////////////////////////////////////////////////////////////
+// FAIDataProviderStructValue
+
+bool FAIDataProviderStructValue::IsMatchingType(UProperty* PropType) const
+{
+	UStructProperty* StructProp = Cast<UStructProperty>(PropType);
+	if (StructProp)
+	{
+		// skip inital "struct " 
+		FString CPPType = StructProp->GetCPPType(nullptr, CPPF_None).Mid(8);
+		return CPPType == StructName;
+	}
+
+	return false;
+}
+
+//////////////////////////////////////////////////////////////////////////
+// FAIDataProviderIntValue
+
+FAIDataProviderIntValue::FAIDataProviderIntValue()
+{
+	PropertyType = UIntProperty::StaticClass();
+}
+
+int32 FAIDataProviderIntValue::GetValue() const
+{
+	int32* PropValue = GetRawValue<int32>();
+	return PropValue ? *PropValue : DefaultValue;
+}
+
+FString FAIDataProviderIntValue::ValueToString() const
+{
+	return TTypeToString<int32>().ToSanitizedString(DefaultValue);
+}
+
+//////////////////////////////////////////////////////////////////////////
+// FAIDataProviderFloatValue
+
+FAIDataProviderFloatValue::FAIDataProviderFloatValue()
+{
+	PropertyType = UFloatProperty::StaticClass();
+}
+
+float FAIDataProviderFloatValue::GetValue() const
+{
+	float* PropValue = GetRawValue<float>();
+	return PropValue ? *PropValue : DefaultValue;
+}
+
+FString FAIDataProviderFloatValue::ValueToString() const
+{
+	return TTypeToString<float>().ToSanitizedString(DefaultValue);
+}
+
+//////////////////////////////////////////////////////////////////////////
+// FAIDataProviderBoolValue
+
+FAIDataProviderBoolValue::FAIDataProviderBoolValue()
+{
+	PropertyType = UBoolProperty::StaticClass();
+}
+
+bool FAIDataProviderBoolValue::GetValue() const
+{
+	bool* PropValue = GetRawValue<bool>();
+	return PropValue ? *PropValue : DefaultValue;
+}
+
+FString FAIDataProviderBoolValue::ValueToString() const
+{
+	return DefaultValue ? GTrue.ToString() : GFalse.ToString();
+}
+
+//////////////////////////////////////////////////////////////////////////
+// UAIDataProvider
+
+UAIDataProvider::UAIDataProvider(const FObjectInitializer& ObjectInitializer) : Super(ObjectInitializer)
+{
+}
+
+void UAIDataProvider::BindData(UObject* Owner, int32 RequestId)
+{
+	// empty in base class
+}
+
+FString UAIDataProvider::ToString(FName PropName) const
+{
+	FString ProviderName = GetClass()->GetName();
+	int32 SplitIdx = 0;
+	const bool bFound = ProviderName.FindChar(TEXT('_'), SplitIdx);
+	if (bFound)
+	{
+		ProviderName = ProviderName.Mid(SplitIdx + 1);
+	}
+
+	ProviderName += TEXT('.');
+	ProviderName += PropName.ToString();
+	return ProviderName;
+}
diff --git a/Engine/Source/Runtime/AIModule/Private/DataProviders/AIDataProvider_QueryParams.cpp b/Engine/Source/Runtime/AIModule/Private/DataProviders/AIDataProvider_QueryParams.cpp
new file mode 100644
--- /dev/null
+++ b/Engine/Source/Runtime/AIModule/Private/DataProviders/AIDataProvider_QueryParams.cpp
@@ -0,0 +1,31 @@
+// Copyright 1998-2015 Epic Games, Inc. All Rights Reserved.
+
+#include "AIModulePrivate.h"
+#include "DataProviders/AIDataProvider_QueryParams.h"
+#include "EnvironmentQuery/EnvQueryManager.h"
+
+UAIDataProvider_QueryParams::UAIDataProvider_QueryParams(const FObjectInitializer& ObjectInitializer) : Super(ObjectInitializer)
+{
+}
+
+void UAIDataProvider_QueryParams::BindData(UObject* Owner, int32 RequestId)
+{
+	UEnvQueryManager* QueryManager = UEnvQueryManager::GetCurrent(Owner);
+	if (QueryManager)
+	{
+		FloatValue = QueryManager->FindNamedParam(RequestId, ParamName);
+		IntValue = *((int32*)&FloatValue);
+		BoolValue = *((bool*)&FloatValue);
+	}
+	else
+	{
+		IntValue = 0;
+		FloatValue = 0.0f;
+		BoolValue = false;
+	}
+}
+
+FString UAIDataProvider_QueryParams::ToString(FName PropName) const
+{
+	return FString::Printf(TEXT("QueryParam.%s"), *ParamName.ToString());
+}
diff --git a/Engine/Source/Runtime/AIModule/Private/EnvironmentQuery/EnvQueryGenerator.cpp b/Engine/Source/Runtime/AIModule/Private/EnvironmentQuery/EnvQueryGenerator.cpp
--- a/Engine/Source/Runtime/AIModule/Private/EnvironmentQuery/EnvQueryGenerator.cpp
+++ b/Engine/Source/Runtime/AIModule/Private/EnvironmentQuery/EnvQueryGenerator.cpp
@@ -8,6 +8,11 @@ UEnvQueryGenerator::UEnvQueryGenerator(const FObjectInitializer& ObjectInitializ
 {
 }
 
+void UEnvQueryGenerator::UpdateGeneratorVersion()
+{
+	VerNum = EnvQueryGeneratorVersion::Latest;
+}
+
 FText UEnvQueryGenerator::GetDescriptionTitle() const
 {
 	return UEnvQueryTypes::GetShortTypeName(this);
@@ -27,3 +32,9 @@ void UEnvQueryGenerator::PostEditChangeProperty(FPropertyChangedEvent& PropertyC
 #endif
 }
 #endif //WITH_EDITOR && USE_EQS_DEBUGGER
+
+void UEnvQueryGenerator::PostLoad()
+{
+	Super::PostLoad();
+	UpdateGeneratorVersion();
+}
\ No newline at end of file
diff --git a/Engine/Source/Runtime/AIModule/Private/EnvironmentQuery/EnvQueryInstance.cpp b/Engine/Source/Runtime/AIModule/Private/EnvironmentQuery/EnvQueryInstance.cpp
--- a/Engine/Source/Runtime/AIModule/Private/EnvironmentQuery/EnvQueryInstance.cpp
+++ b/Engine/Source/Runtime/AIModule/Private/EnvironmentQuery/EnvQueryInstance.cpp
@@ -317,7 +317,6 @@ FEnvQueryInstance::ItemIterator::ItemIterator(const UEnvQueryTest* QueryTest, FE
 	, CurrentItem(StartingItemIndex != INDEX_NONE ? StartingItemIndex : QueryInstance.CurrentTestStartingItem)
 {
 	Deadline = QueryInstance.TimeLimit > 0.0 ? (FPlatformTime::Seconds() + QueryInstance.TimeLimit) : -1.0;
-	bDiscardFailed = QueryTest && QueryTest->bDiscardFailedItems;
 	// it's possible item 'CurrentItem' has been already discarded. Find a valid starting index
 	--CurrentItem;
 	FindNextValidIndex();
diff --git a/Engine/Source/Runtime/AIModule/Private/EnvironmentQuery/EnvQueryManager.cpp b/Engine/Source/Runtime/AIModule/Private/EnvironmentQuery/EnvQueryManager.cpp
--- a/Engine/Source/Runtime/AIModule/Private/EnvironmentQuery/EnvQueryManager.cpp
+++ b/Engine/Source/Runtime/AIModule/Private/EnvironmentQuery/EnvQueryManager.cpp
@@ -364,14 +364,6 @@ namespace EnvQueryTestSort
 				{
 					return true;
 				}
-
-				// conditions with weights before pure conditions
-				const bool bNoWeightA = (TestA.WeightModifier >= EEnvTestWeight::Constant);
-				const bool bNoWeightB = (TestB.WeightModifier >= EEnvTestWeight::Constant);
-				if (!bNoWeightA && bNoWeightB)
-				{
-					return true;
-				}
 			}
 			else
 			{
@@ -576,6 +568,23 @@ UEnvQueryContext* UEnvQueryManager::PrepareLocalContext(TSubclassOf<UEnvQueryCon
 	return LocalContext;
 }
 
+float UEnvQueryManager::FindNamedParam(int32 QueryId, FName ParamName) const
+{
+	float ParamValue = 0.0f;
+
+	for (int32 QueryIndex = 0; QueryIndex < RunningQueries.Num(); QueryIndex++)
+	{
+		const TSharedPtr<FEnvQueryInstance>& QueryInstance = RunningQueries[QueryIndex];
+		if (QueryInstance->QueryID == QueryId)
+		{
+			ParamValue = QueryInstance->NamedParams.FindRef(ParamName);
+			break;
+		}
+	}
+
+	return ParamValue;
+}
+
 
 //----------------------------------------------------------------------//
 // FEQSDebugger
diff --git a/Engine/Source/Runtime/AIModule/Private/EnvironmentQuery/EnvQueryTest.cpp b/Engine/Source/Runtime/AIModule/Private/EnvironmentQuery/EnvQueryTest.cpp
--- a/Engine/Source/Runtime/AIModule/Private/EnvironmentQuery/EnvQueryTest.cpp
+++ b/Engine/Source/Runtime/AIModule/Private/EnvironmentQuery/EnvQueryTest.cpp
@@ -14,20 +14,19 @@
 UEnvQueryTest::UEnvQueryTest(const FObjectInitializer& ObjectInitializer) : Super(ObjectInitializer)
 {
 	TestPurpose = EEnvTestPurpose::FilterAndScore;
-	FilterType = EEnvTestFilterType::Range;
-	BoolFilter.Value = true;
-	bFormatUpdated = false;
-	Condition = EEnvTestCondition::NoCondition;	// DEPRECATED
-	Weight.Value = 1.0f;						// DEPRECATED
-	WeightModifier = EEnvTestWeight::None;		// DEPRECATED
 	Cost = EEnvTestCost::Low;
+	FilterType = EEnvTestFilterType::Range;
+	ScoringEquation = EEnvTestScoreEquation::Linear;
 	ClampMinType = EEnvQueryTestClamping::None;
 	ClampMaxType = EEnvQueryTestClamping::None;
-	ScoringEquation = EEnvTestScoreEquation::Linear;
-	
-	bNormalizeFromZero = false;					// DEPRECATED
+	BoolValue.DefaultValue = true;
+	ScoringFactor.DefaultValue = 1.0f;
+
 	bWorkOnFloatValues = true;
-	bDiscardFailedItems = true;					// DEPRECATED
+
+	// keep deprecated properties initialized
+	BoolFilter.Value = true;
+	Weight.Value = 1.0f;
 }
 
 void UEnvQueryTest::NormalizeItemScores(FEnvQueryInstance& QueryInstance)
@@ -37,47 +36,32 @@ void UEnvQueryTest::NormalizeItemScores(FEnvQueryInstance& QueryInstance)
 		return;
 	}
 
-	float WeightValue = 0.0f;
-	if (!QueryInstance.GetParamValue(Weight, WeightValue, TEXT("Weight")))
-	{
-		return;
-	}
+	ScoringFactor.BindData(QueryInstance.Owner.Get(), QueryInstance.QueryID);
+	float ScoringFactorValue = ScoringFactor.GetValue();
 
 	float MinScore = 0;
 	float MaxScore = -BIG_NUMBER;
 
 	if (ClampMinType == EEnvQueryTestClamping::FilterThreshold)
 	{
-		bool bSuccess = QueryInstance.GetParamValue(FloatFilterMin, MinScore, TEXT("FloatFilterMin"));
-		if (!bSuccess)
-		{
-			UE_LOG(LogEQS, Warning, TEXT("Unable to get FloatFilterMin parameter value from EnvQueryInstance %s"), FloatFilterMin.IsNamedParam() ? *FloatFilterMin.ParamName.ToString() : TEXT("<No name specified>"));
-		}
+		FloatValueMin.BindData(QueryInstance.Owner.Get(), QueryInstance.QueryID);
+		MinScore = FloatValueMin.GetValue();
 	}
 	else if (ClampMinType == EEnvQueryTestClamping::SpecifiedValue)
 	{
-		bool bSuccess = QueryInstance.GetParamValue(ScoreClampingMin, MinScore, TEXT("ScoreClampingMin"));
-		if (!bSuccess)
-		{
-			UE_LOG(LogEQS, Warning, TEXT("Unable to get ClampMinType parameter value from EnvQueryInstance %s"), ScoreClampingMin.IsNamedParam() ? *ScoreClampingMin.ParamName.ToString() : TEXT("<No name specified>"));
-		}
+		ScoreClampMin.BindData(QueryInstance.Owner.Get(), QueryInstance.QueryID);
+		MinScore = ScoreClampMin.GetValue();
 	}
 
 	if (ClampMaxType == EEnvQueryTestClamping::FilterThreshold)
 	{
-		bool bSuccess = QueryInstance.GetParamValue(FloatFilterMax, MaxScore, TEXT("FloatFilterMax"));
-		if (!bSuccess)
-		{
-			UE_LOG(LogEQS, Warning, TEXT("Unable to get FloatFilterMax parameter value from EnvQueryInstance %s"), FloatFilterMax.IsNamedParam() ? *FloatFilterMax.ParamName.ToString() : TEXT("<No name specified>"));
-		}
+		FloatValueMax.BindData(QueryInstance.Owner.Get(), QueryInstance.QueryID);
+		MaxScore = FloatValueMax.GetValue();
 	}
 	else if (ClampMaxType == EEnvQueryTestClamping::SpecifiedValue)
 	{
-		bool bSuccess = QueryInstance.GetParamValue(ScoreClampingMax, MaxScore, TEXT("ScoreClampingMax"));
-		if (!bSuccess)
-		{
-			UE_LOG(LogEQS, Warning, TEXT("Unable to get ScoreClampingMax parameter value from EnvQueryInstance %s"), ScoreClampingMax.IsNamedParam() ? *ScoreClampingMax.ParamName.ToString() : TEXT("<No name specified>"));
-		}
+		ScoreClampMax.BindData(QueryInstance.Owner.Get(), QueryInstance.QueryID);
+		MaxScore = ScoreClampMax.GetValue();
 	}
 
 	FEnvQueryItemDetails* DetailInfo = QueryInstance.ItemDetails.GetData();
@@ -109,10 +93,6 @@ void UEnvQueryTest::NormalizeItemScores(FEnvQueryInstance& QueryInstance)
 	}
 
 	DetailInfo = QueryInstance.ItemDetails.GetData();
-	if (bNormalizeFromZero && MinScore > 0.0f)
-	{
-		MinScore = 0.0f;
-	}
 
 	if (MinScore != MaxScore)
 	{
@@ -135,7 +115,7 @@ void UEnvQueryTest::NormalizeItemScores(FEnvQueryInstance& QueryInstance)
 				switch (ScoringEquation)
 				{
 					case EEnvTestScoreEquation::Linear:
-						WeightedScore = WeightValue * NormalizedScoreForEquation;
+						WeightedScore = ScoringFactorValue * NormalizedScoreForEquation;
 						break;
 
 					case EEnvTestScoreEquation::InverseLinear:
@@ -145,17 +125,17 @@ void UEnvQueryTest::NormalizeItemScores(FEnvQueryInstance& QueryInstance)
 						// to add that flag later, we'll need to remove this option, since it should just be "mirror
 						// curve" plus "Linear".
 						float InverseNormalizedScore = (1.0f - NormalizedScoreForEquation);
-						WeightedScore = WeightValue * InverseNormalizedScore;
+						WeightedScore = ScoringFactorValue * InverseNormalizedScore;
 						break;
 					}
 
 					case EEnvTestScoreEquation::Square:
-						WeightedScore = WeightValue * (NormalizedScoreForEquation * NormalizedScoreForEquation);
+						WeightedScore = ScoringFactorValue * (NormalizedScoreForEquation * NormalizedScoreForEquation);
 						break;
 
 					case EEnvTestScoreEquation::Constant:
 						// I know, it's not "constant".  It's "Constant, or zero".  The tooltip should explain that.
-						WeightedScore = (NormalizedScoreForEquation > 0) ? WeightValue : 0.0f;
+						WeightedScore = (NormalizedScoreForEquation > 0) ? ScoringFactorValue : 0.0f;
 						break;
 						
 					default:
@@ -216,83 +196,14 @@ void UEnvQueryTest::PostLoad()
 {
 	Super::PostLoad();
 
-	if (bFormatUpdated)
-	{
-		// Backwards compatibility already handled.
-		return;
-	}
-
-	switch (Condition)
-	{
-		case EEnvTestCondition::AtLeast:
-			FloatFilterMin = FloatFilter;
-			FilterType = EEnvTestFilterType::Minimum;
-			break;
-		
-		case EEnvTestCondition::UpTo:
-			FloatFilterMax = FloatFilter;
-			FilterType = EEnvTestFilterType::Maximum;
-			break;
-		
-		case EEnvTestCondition::NoCondition:
-			TestPurpose = EEnvTestPurpose::Score;
-			break;
-
-		case EEnvTestCondition::Match:
-			FilterType = EEnvTestFilterType::Match;
-			break;
-
-		default:
-			UE_LOG(LogEQS, Error, TEXT("Invalid Condition type in UEnvQueryTest::PostLoad!"));
-			break;
-	}
-
-	if ((WeightModifier != EEnvTestWeight::Skip) && !Weight.IsNamedParam() && (Weight.Value == 0.f))
-	{	// We ought to be skipping the test if weight is 0!  Zero makes it pointless!
-		WeightModifier = EEnvTestWeight::Skip;
-	}
-
-	if (WeightModifier != EEnvTestWeight::Skip && !bWorkOnFloatValues)
+	if (VerNum < EnvQueryTestVersion::DataProviders)
 	{
-		// Working on booleans, so MUST use constant value!
-		WeightModifier = EEnvTestWeight::Constant;
-	}
-
-	switch (WeightModifier)
-	{
-		case EEnvTestWeight::None:
-			ScoringEquation = EEnvTestScoreEquation::Linear;
-			break;
-
-		case EEnvTestWeight::Constant:
-			ScoringEquation = EEnvTestScoreEquation::Constant;
-			break;
-
-		case EEnvTestWeight::Inverse:
-			ScoringEquation = EEnvTestScoreEquation::InverseLinear;
-			break;
-
-		case EEnvTestWeight::Square:
-			ScoringEquation = EEnvTestScoreEquation::Square;
-			break;
-
-		case EEnvTestWeight::Absolute:
-			UE_LOG(LogEQS, Warning, TEXT("Absolute weight is no longer supported!  Sticking with default weighting."));
-			// falling back to Linear
-			ScoringEquation = EEnvTestScoreEquation::Linear;
-			break;
-
-		case EEnvTestWeight::Skip:
-			TestPurpose = EEnvTestPurpose::Filter;
-			if (Condition == EEnvTestCondition::NoCondition)
-			{
-				UE_LOG(LogEQS, Warning, TEXT("Test was set to neither filter nor score!  Now setting to filter!"));	
-			}
-			break;
-
-		default:
-			UE_LOG(LogEQS, Warning, TEXT("Invalid Weight Modifier type in UEnvQueryTest::PostLoad!"));
-			break;
+		BoolFilter.Convert(this, BoolValue);
+		FloatFilterMin.Convert(this, FloatValueMin);
+		FloatFilterMax.Convert(this, FloatValueMax);
+		ScoreClampingMin.Convert(this, ScoreClampMin);
+		ScoreClampingMax.Convert(this, ScoreClampMax);
+		Weight.Convert(this, ScoringFactor);
 	}
 
 	UpdateTestVersion();
@@ -300,7 +211,7 @@ void UEnvQueryTest::PostLoad()
 
 void UEnvQueryTest::UpdateTestVersion()
 {
-	bFormatUpdated = true;
+	VerNum = EnvQueryTestVersion::Latest;
 }
 
 #if WITH_EDITOR && USE_EQS_DEBUGGER
@@ -315,108 +226,97 @@ void UEnvQueryTest::PostEditChangeProperty(FPropertyChangedEvent& PropertyChange
 
 FText UEnvQueryTest::DescribeFloatTestParams() const
 {
-	FText ParamDesc;
-
-	if (Condition != EEnvTestCondition::NoCondition)
+	FText FilterDesc;
+	if (IsFiltering())
 	{
-		FFormatNamedArguments Args;
-		Args.Add(TEXT("Filter"), FText::FromString(UEnvQueryTypes::DescribeFloatParam(FloatFilter)));
-
-		FText InvalidConditionDesc = LOCTEXT("InvalidCondition", "invalid condition");
-		ParamDesc =
-			(Condition == EEnvTestCondition::AtLeast) ? FText::Format(LOCTEXT("AtLeastFiltered", "at least {Filter}"), Args) :
-			(Condition == EEnvTestCondition::UpTo) ? FText::Format(LOCTEXT("UpToFiltered", "up to {Filter}"), Args) :
-			InvalidConditionDesc;
+		switch (FilterType)
+		{
+		case EEnvTestFilterType::Minimum:
+			FilterDesc = FText::Format(LOCTEXT("FilterAtLeast", "at least {0}"), FText::FromString(FloatValueMin.ToString()));
+			break;
+
+		case EEnvTestFilterType::Maximum:
+			FilterDesc = FText::Format(LOCTEXT("FilterUpTo", "up to {0}"), FText::FromString(FloatValueMax.ToString()));
+			break;
+
+		case EEnvTestFilterType::Range:
+			FilterDesc = FText::Format(LOCTEXT("FilterBetween", "between {0} and {1}"),
+				FText::FromString(FloatValueMin.ToString()), FText::FromString(FloatValueMax.ToString()));
+			break;
+
+		default:
+			break;
+		}
 	}
 
-	FFormatNamedArguments Args;
-	Args.Add(TEXT("ParmDesc"), ParamDesc);
+	FNumberFormattingOptions NumberFormattingOptions;
+	NumberFormattingOptions.MaximumFractionalDigits = 2;
 
+	FText ScoreDesc;
 	if (!IsScoring())
 	{
-		ParamDesc = FText::Format(LOCTEXT("ParmDescWithDontScore", "{ParmDesc}, don't score"), Args);
+		ScoreDesc = LOCTEXT("DontScore", "don't score");
 	}
 	else if (ScoringEquation == EEnvTestScoreEquation::Constant)
 	{
-		FText WeightDesc = Weight.IsNamedParam() ? FText::FromName(Weight.ParamName) : FText::Format( LOCTEXT("ConstantWeightDescPattern", "x{0}"), FText::AsNumber(FMath::Abs(Weight.Value)) );
-		Args.Add(TEXT("WeightDesc"), WeightDesc);
-		ParamDesc = FText::Format(LOCTEXT("ParmDescWithConstantWeight", "{ParmDesc}, constant weight [{WeightDesc}]"), Args);
+		FText FactorDesc = ScoringFactor.IsDynamic() ?
+			FText::FromString(ScoringFactor.ToString()) :
+			FText::Format(FText::FromString("x{0}"), FText::AsNumber(FMath::Abs(ScoringFactor.DefaultValue), &NumberFormattingOptions));
+
+		ScoreDesc = FText::Format(FText::FromString("{0} [{1}]"), LOCTEXT("ScoreConstant", "constant score"), FactorDesc);
 	}
-	else if (Weight.IsNamedParam())
+	else if (ScoringFactor.IsDynamic())
 	{
-		Args.Add(TEXT("WeightName"), FText::FromName(Weight.ParamName));
-		ParamDesc = FText::Format(LOCTEXT("ParmDescWithWeight", "{ParmDesc}, weight: {WeightName}"), Args);
+		ScoreDesc = FText::Format(FText::FromString("{0}: {1}"), LOCTEXT("ScoreFactor", "score factor"), FText::FromString(ScoringFactor.ToString()));
 	}
 	else
 	{
-		FNumberFormattingOptions NumberFormattingOptions;
-		NumberFormattingOptions.MaximumFractionalDigits = 2;
-
-		Args.Add(TEXT("WeightCondition"), (Weight.Value > 0) ? LOCTEXT("Greater", "greater") : LOCTEXT("Lesser", "lesser"));
-		Args.Add(TEXT("WeightValue"), FText::AsNumber(FMath::Abs(Weight.Value), &NumberFormattingOptions));
-		ParamDesc = FText::Format(LOCTEXT("DescriptionPreferWeightValue", "{ParmDesc}, prefer {WeightCondition} [x{WeightValue}]"), Args);
+		FText ScoreSignDesc = (ScoringFactor.DefaultValue > 0) ? LOCTEXT("Greater", "greater") : LOCTEXT("Lesser", "lesser");
+		FText ScoreValueDesc = FText::AsNumber(FMath::Abs(ScoringFactor.DefaultValue), &NumberFormattingOptions);
+		ScoreDesc = FText::Format(FText::FromString("{0} {1} [x{2}]"), LOCTEXT("ScorePrefer", "prefer"), ScoreSignDesc, ScoreValueDesc);
 	}
 
-	return ParamDesc;
+	return FilterDesc.IsEmpty() ? ScoreDesc : FText::Format(FText::FromString("{0}, {1}"), FilterDesc, ScoreDesc);
 }
 
 FText UEnvQueryTest::DescribeBoolTestParams(const FString& ConditionDesc) const
 {
-	FText ParamDesc;
-
-	if (Condition != EEnvTestCondition::NoCondition)
+	FText FilterDesc;
+	if (IsFiltering() && FilterType == EEnvTestFilterType::Match)
 	{
-		FFormatNamedArguments Args;
-		Args.Add(TEXT("ConditionDesc"), FText::FromString(ConditionDesc));
-
-		if(BoolFilter.IsNamedParam())
-		{
-			Args.Add(TEXT("Filter"), FText::FromString(UEnvQueryTypes::DescribeBoolParam(BoolFilter)));
-		}
-		else
-		{
-			Args.Add(TEXT("Filter"), BoolFilter.Value ? FText::GetEmpty() : LOCTEXT("NotWithSpace", "not "));
-		}
-		FText InvalidConditionDesc = LOCTEXT("InvalidCondition", "invalid condition");
-		ParamDesc =
-			(Condition == EEnvTestCondition::Match) ?
-				(BoolFilter.IsNamedParam() ?
-					FText::Format(LOCTEXT("RequiresCondition", "require {ConditionDesc}: {Filter}"), Args) :
-					FText::Format(LOCTEXT("RequiresFilter", "require {Filter}{ConditionDesc}%s"), Args)
-				) :
-			InvalidConditionDesc;
+		FilterDesc = BoolValue.IsDynamic() ?
+			FText::Format(FText::FromString("{0} {1}: {2}"), LOCTEXT("FilterRequire", "require"), FText::FromString(ConditionDesc), FText::FromString(BoolValue.ToString())) :
+			FText::Format(FText::FromString("{0} {1}{2}"), LOCTEXT("FilterRequire", "require"), BoolValue.DefaultValue ? FText::GetEmpty() : LOCTEXT("NotWithSpace", "not "), FText::FromString(ConditionDesc));
 	}
 
-	FFormatNamedArguments Args;
-	Args.Add(TEXT("ParmDesc"), ParamDesc);
+	FNumberFormattingOptions NumberFormattingOptions;
+	NumberFormattingOptions.MaximumFractionalDigits = 2;
 
+	FText ScoreDesc;
 	if (!IsScoring())
 	{
-		ParamDesc = FText::Format(LOCTEXT("ParmDescWithDontScore", "{ParmDesc}, don't score"), Args);
+		ScoreDesc = LOCTEXT("DontScore", "don't score");
 	}
 	else if (ScoringEquation == EEnvTestScoreEquation::Constant)
 	{
-		FText WeightDesc = Weight.IsNamedParam() ? FText::FromName(Weight.ParamName) : FText::Format(LOCTEXT("ConstantWeightDescPattern", "x{0}"), FText::AsNumber(FMath::Abs(Weight.Value)));
-		Args.Add(TEXT("WeightDesc"), WeightDesc);
-		ParamDesc = FText::Format(LOCTEXT("ParmDescWithConstantWeight", "{ParmDesc}, constant weight [{WeightDesc}]"), Args);
+		FText FactorDesc = ScoringFactor.IsDynamic() ?
+			FText::FromString(ScoringFactor.ToString()) :
+			FText::Format(FText::FromString("x{0}"), FText::AsNumber(FMath::Abs(ScoringFactor.DefaultValue), &NumberFormattingOptions));
+
+		ScoreDesc = FText::Format(FText::FromString("{0} [{1}]"), LOCTEXT("ScoreConstant", "constant score"), FactorDesc);
 	}
-	else if (Weight.IsNamedParam())
+	else if (ScoringFactor.IsDynamic())
 	{
-		Args.Add(TEXT("WeightName"), FText::FromName(Weight.ParamName));
-		ParamDesc = FText::Format(LOCTEXT("ParmDescWithWeight", "{ParmDesc}, weight: {WeightName}"), Args);
+		ScoreDesc = FText::Format(FText::FromString("{0}: {1}"), LOCTEXT("ScoreFactor", "score factor"), FText::FromString(ScoringFactor.ToString()));
 	}
 	else
 	{
-		FNumberFormattingOptions NumberFormattingOptions;
-		NumberFormattingOptions.MaximumFractionalDigits = 2;
-
-		Args.Add(TEXT("ConditionDesc"), FText::FromString(ConditionDesc));
-		Args.Add(TEXT("WeightCondition"), (Weight.Value > 0) ? FText::GetEmpty() : LOCTEXT("NotWithSpace", "not "));
-		Args.Add(TEXT("WeightValue"), FText::AsNumber(FMath::Abs(Weight.Value), &NumberFormattingOptions));
-		ParamDesc = FText::Format(LOCTEXT("ParmDescWithPreference", "{ParmDesc}, prefer {WeightCondition}{ConditionDesc} [x{WeightValue}]"), Args);
+		FText ScoreSignDesc = (ScoringFactor.DefaultValue > 0) ? FText::GetEmpty() : LOCTEXT("NotWithSpace", "not ");
+		FText ScoreValueDesc = FText::AsNumber(FMath::Abs(ScoringFactor.DefaultValue), &NumberFormattingOptions);
+		ScoreDesc = FText::Format(FText::FromString("{0} {1}{2} [x{3}]"), LOCTEXT("ScorePrefer", "prefer"), ScoreSignDesc, FText::FromString(ConditionDesc), ScoreValueDesc);
 	}
 
-	return ParamDesc;
+	return FilterDesc.IsEmpty() ? ScoreDesc : FText::Format(FText::FromString("{0}, {1}"), FilterDesc, ScoreDesc);
 }
 
 void UEnvQueryTest::SetWorkOnFloatValues(bool bWorkOnFloats)
diff --git a/Engine/Source/Runtime/AIModule/Private/EnvironmentQuery/EnvQueryTypes.cpp b/Engine/Source/Runtime/AIModule/Private/EnvironmentQuery/EnvQueryTypes.cpp
--- a/Engine/Source/Runtime/AIModule/Private/EnvironmentQuery/EnvQueryTypes.cpp
+++ b/Engine/Source/Runtime/AIModule/Private/EnvironmentQuery/EnvQueryTypes.cpp
@@ -1,6 +1,7 @@
 // Copyright 1998-2015 Epic Games, Inc. All Rights Reserved.
 
 #include "AIModulePrivate.h"
+#include "DataProviders/AIDataProvider_QueryParams.h"
 #include "EnvironmentQuery/EnvQueryTypes.h"
 #include "EnvironmentQuery/EnvQueryContext.h"
 #include "EnvironmentQuery/Items/EnvQueryItemType.h"
@@ -62,36 +63,6 @@ FText UEnvQueryTypes::DescribeContext(TSubclassOf<UEnvQueryContext> ContextClass
 	return GetShortTypeName(ContextClass);
 }
 
-FString UEnvQueryTypes::DescribeIntParam(const FEnvIntParam& Param)
-{
-	if (Param.IsNamedParam())
-	{
-		return Param.ParamName.ToString();
-	}
-
-	return FString::Printf(TEXT("%d"), Param.Value);
-}
-
-FString UEnvQueryTypes::DescribeFloatParam(const FEnvFloatParam& Param)
-{
-	if (Param.IsNamedParam())
-	{
-		return Param.ParamName.ToString();
-	}
-
-	return FString::Printf(TEXT("%.2f"), Param.Value);
-}
-
-FString UEnvQueryTypes::DescribeBoolParam(const FEnvBoolParam& Param)
-{
-	if (Param.IsNamedParam())
-	{
-		return Param.ParamName.ToString();
-	}
-
-	return Param.Value ? TEXT("true") : TEXT("false");
-}
-
 FText FEnvDirection::ToText() const
 {
 	if(DirMode == EEnvDirection::TwoPoints)
@@ -225,6 +196,43 @@ void FEnvTraceData::SetNavmeshOnly()
 	bCanDisableTrace = false;
 }
 
+void FEnvBoolParam::Convert(UObject* Owner, FAIDataProviderBoolValue& ValueProvider)
+{
+	ValueProvider.DefaultValue = Value;
+	if (IsNamedParam())
+	{
+		UAIDataProvider_QueryParams* ParamsProvider = NewObject<UAIDataProvider_QueryParams>(Owner);
+		ParamsProvider->ParamName = ParamName;
+		ValueProvider.DataBinding = ParamsProvider;
+		ValueProvider.DataField = GET_MEMBER_NAME_CHECKED(UAIDataProvider_QueryParams, BoolValue);
+	}
+}
+
+void FEnvIntParam::Convert(UObject* Owner, FAIDataProviderIntValue& ValueProvider)
+{
+	ValueProvider.DefaultValue = Value;
+	if (IsNamedParam())
+	{
+		UAIDataProvider_QueryParams* ParamsProvider = NewObject<UAIDataProvider_QueryParams>(Owner);
+		ParamsProvider->ParamName = ParamName;
+		ValueProvider.DataBinding = ParamsProvider;
+		ValueProvider.DataField = GET_MEMBER_NAME_CHECKED(UAIDataProvider_QueryParams, IntValue);
+	}
+}
+
+void FEnvFloatParam::Convert(UObject* Owner, FAIDataProviderFloatValue& ValueProvider)
+{
+	ValueProvider.DefaultValue = Value;
+	if (IsNamedParam())
+	{
+		UAIDataProvider_QueryParams* ParamsProvider = NewObject<UAIDataProvider_QueryParams>(Owner);
+		ParamsProvider->ParamName = ParamName;
+		ValueProvider.DataBinding = ParamsProvider;
+		ValueProvider.DataField = GET_MEMBER_NAME_CHECKED(UAIDataProvider_QueryParams, FloatValue);
+	}
+}
+
+
 //----------------------------------------------------------------------//
 // namespace FEQSHelpers
 //----------------------------------------------------------------------//
diff --git a/Engine/Source/Runtime/AIModule/Private/EnvironmentQuery/Generators/EnvQueryGenerator_ActorsOfClass.cpp b/Engine/Source/Runtime/AIModule/Private/EnvironmentQuery/Generators/EnvQueryGenerator_ActorsOfClass.cpp
--- a/Engine/Source/Runtime/AIModule/Private/EnvironmentQuery/Generators/EnvQueryGenerator_ActorsOfClass.cpp
+++ b/Engine/Source/Runtime/AIModule/Private/EnvironmentQuery/Generators/EnvQueryGenerator_ActorsOfClass.cpp
@@ -13,19 +13,20 @@ UEnvQueryGenerator_ActorsOfClass::UEnvQueryGenerator_ActorsOfClass(const FObject
 {
 	SearchCenter = UEnvQueryContext_Querier::StaticClass();
 	ItemType = UEnvQueryItemType_Actor::StaticClass();
-	Radius.Value = 500.0f;
+	SearchRadius.DefaultValue = 500.0f;
 	SearchedActorClass = AActor::StaticClass();
+
+	// keep deprecated properties initialized
+	Radius.Value = 500.0f;
 }
 
 void UEnvQueryGenerator_ActorsOfClass::GenerateItems(FEnvQueryInstance& QueryInstance) const
 {
-	float RadiusValue = 0.0f;
-	float DensityValue = 0.0f;
+	SearchRadius.BindData(QueryInstance.Owner.Get(), QueryInstance.QueryID);
+	float RadiusValue = SearchRadius.GetValue();
 
 	UWorld* World = GEngine->GetWorldFromContextObject(QueryInstance.Owner.Get());
-	// @TODO add some logging here
-	if (World == NULL || QueryInstance.GetParamValue(Radius, RadiusValue, TEXT("Radius")) == false
-		|| SearchedActorClass == NULL)
+	if (World == NULL || SearchedActorClass == NULL)
 	{
 		return;
 	}
@@ -61,11 +62,21 @@ FText UEnvQueryGenerator_ActorsOfClass::GetDescriptionTitle() const
 FText UEnvQueryGenerator_ActorsOfClass::GetDescriptionDetails() const
 {
 	FFormatNamedArguments Args;
-	Args.Add(TEXT("Radius"), FText::FromString(UEnvQueryTypes::DescribeFloatParam(Radius)));
+	Args.Add(TEXT("Radius"), FText::FromString(SearchRadius.ToString()));
 
 	FText Desc = FText::Format(LOCTEXT("ActorsOfClassDescription", "radius: {Radius}"), Args);
 
 	return Desc;
 }
 
+void UEnvQueryGenerator_ActorsOfClass::PostLoad()
+{
+	if (VerNum < EnvQueryGeneratorVersion::DataProviders)
+	{
+		Radius.Convert(this, SearchRadius);
+	}
+
+	Super::PostLoad();
+}
+
 #undef LOCTEXT_NAMESPACE
\ No newline at end of file
diff --git a/Engine/Source/Runtime/AIModule/Private/EnvironmentQuery/Generators/EnvQueryGenerator_BlueprintBase.cpp b/Engine/Source/Runtime/AIModule/Private/EnvironmentQuery/Generators/EnvQueryGenerator_BlueprintBase.cpp
--- a/Engine/Source/Runtime/AIModule/Private/EnvironmentQuery/Generators/EnvQueryGenerator_BlueprintBase.cpp
+++ b/Engine/Source/Runtime/AIModule/Private/EnvironmentQuery/Generators/EnvQueryGenerator_BlueprintBase.cpp
@@ -16,7 +16,6 @@ UEnvQueryGenerator_BlueprintBase::UEnvQueryGenerator_BlueprintBase(const FObject
 {
 	Context = UEnvQueryContext_Querier::StaticClass();
 	ItemType = UEnvQueryItemType_Actor::StaticClass();
-	Radius.Value = 1000.0f;
 	GeneratedItemType = UEnvQueryItemType_Actor::StaticClass();
 }
 
@@ -33,12 +32,6 @@ UWorld* UEnvQueryGenerator_BlueprintBase::GetWorld() const
 
 void UEnvQueryGenerator_BlueprintBase::GenerateItems(FEnvQueryInstance& QueryInstance) const
 {
-	float RadiusValue = 0.f;
-	if (QueryInstance.GetParamValue(Radius, RadiusValue, TEXT("Radius")) == false)
-	{
-		return;
-	}
-
 	TArray<FVector> ContextLocations;
 	QueryInstance.PrepareContext(Context, ContextLocations);
 
diff --git a/Engine/Source/Runtime/AIModule/Private/EnvironmentQuery/Generators/EnvQueryGenerator_OnCircle.cpp b/Engine/Source/Runtime/AIModule/Private/EnvironmentQuery/Generators/EnvQueryGenerator_OnCircle.cpp
--- a/Engine/Source/Runtime/AIModule/Private/EnvironmentQuery/Generators/EnvQueryGenerator_OnCircle.cpp
+++ b/Engine/Source/Runtime/AIModule/Private/EnvironmentQuery/Generators/EnvQueryGenerator_OnCircle.cpp
@@ -100,12 +100,12 @@ UEnvQueryGenerator_OnCircle::UEnvQueryGenerator_OnCircle(const FObjectInitialize
 {
 	CircleCenter = UEnvQueryContext_Querier::StaticClass();
 	ItemType = UEnvQueryItemType_Point::StaticClass();
-	Radius.Value = 1000.0f;
-	ItemSpacing.Value = 50.0f;
+	CircleRadius.DefaultValue = 1000.0f;
+	SpaceBetween.DefaultValue = 50.0f;
 	ArcDirection.DirMode = EEnvDirection::TwoPoints;
 	ArcDirection.LineFrom = UEnvQueryContext_Querier::StaticClass();
 	ArcDirection.Rotation = UEnvQueryContext_Querier::StaticClass();
-	Angle.Value = 360.f;
+	ArcAngle.DefaultValue = 360.f;
 	AngleRadians = FMath::DegreesToRadians(360.f);
 
 	TraceData.bCanProjectDown = false;
@@ -113,17 +113,30 @@ UEnvQueryGenerator_OnCircle::UEnvQueryGenerator_OnCircle(const FObjectInitialize
 	TraceData.TraceMode = EEnvQueryTrace::Navigation;
 
 	ProjectionData.TraceMode = EEnvQueryTrace::None;
+
+	// keep deprecated properties initialized
+	Radius.Value = 1000.0f;
+	ItemSpacing.Value = 50.0f;
+	Angle.Value = 360.0f;
 }
 
 void UEnvQueryGenerator_OnCircle::PostLoad()
 {
+	if (VerNum < EnvQueryGeneratorVersion::DataProviders)
+	{
+		Radius.Convert(this, CircleRadius);
+		ItemSpacing.Convert(this, SpaceBetween);
+		Angle.Convert(this, ArcAngle);
+	}
+
 	Super::PostLoad();
 
-	if (Angle.Value <= 0)
+	if (ArcAngle.DefaultValue <= 0)
 	{
-		Angle.Value = 360.f;
+		ArcAngle.DefaultValue = 360.f;
 	}
-	AngleRadians = FMath::DegreesToRadians(Angle.Value);
+
+	AngleRadians = FMath::DegreesToRadians(ArcAngle.DefaultValue);
 }
 
 FVector UEnvQueryGenerator_OnCircle::CalcDirection(FEnvQueryInstance& QueryInstance) const
@@ -177,22 +190,20 @@ FVector UEnvQueryGenerator_OnCircle::CalcDirection(FEnvQueryInstance& QueryInsta
 
 void UEnvQueryGenerator_OnCircle::GenerateItems(FEnvQueryInstance& QueryInstance) const
 {
-	float AngleDegree = 360.f;
-	float RadiusValue = 0.f;
-	float ItemSpace = 1.0f;
-
-	if (QueryInstance.GetParamValue(Angle, AngleDegree, TEXT("Angle")) == false
-		|| QueryInstance.GetParamValue(Radius, RadiusValue, TEXT("Radius")) == false
-		|| QueryInstance.GetParamValue(ItemSpacing, ItemSpace, TEXT("ItemSpacing")) == false
-		|| (AngleDegree <= 0.f) || (AngleDegree > 360.f)
-		|| (RadiusValue <= 0.f)
-		|| (ItemSpace <= 0.f)
-	   )
+	CircleRadius.BindData(QueryInstance.Owner.Get(), QueryInstance.QueryID);
+	SpaceBetween.BindData(QueryInstance.Owner.Get(), QueryInstance.QueryID);
+	ArcAngle.BindData(QueryInstance.Owner.Get(), QueryInstance.QueryID);
+
+	float AngleDegree = ArcAngle.GetValue();
+	float RadiusValue = CircleRadius.GetValue();
+	float ItemSpace = SpaceBetween.GetValue();
+
+	if ((AngleDegree <= 0.f) || (AngleDegree > 360.f) || (RadiusValue <= 0.f) || (ItemSpace <= 0.f))
 	{
 		return;
 	}
 
-	AngleRadians = FMath::DegreesToRadians(Angle.Value);
+	AngleRadians = FMath::DegreesToRadians(AngleDegree);
 
 	// first generate points on a circle
 	const float CircumferenceLength = 2.f * PI * RadiusValue;
@@ -204,63 +215,28 @@ void UEnvQueryGenerator_OnCircle::GenerateItems(FEnvQueryInstance& QueryInstance
 	FVector StartDirection = CalcDirection(QueryInstance);
 	StartDirection = StartDirection.RotateAngleAxis(-AngleDegree/2, FVector::UpVector) * RadiusValue;
 	
-	TSharedPtr<FVectorAndDataContainer> CenterLocationCandidates;
-	PrepareGeneratorContext(QueryInstance, CenterLocationCandidates);
-	check(CenterLocationCandidates.Get() != NULL);
-
-	int32 NumCenterLocations = CenterLocationCandidates.Get()->GetNumEntries();
-	if (NumCenterLocations > 0)
-	{
-		for (int i = 0; i < NumCenterLocations; ++i)
-		{
-			GenerateItemsForCircle(CenterLocationCandidates, i, CenterLocationCandidates.Get()->GetVector(i),
-									StartDirection, StepsCount, AngleStep, QueryInstance);
-		}
-	}
-// 	else
-// 	{
-// 		// No center locations were generated, so we can't generate any item locations based on those centers.
-// 		// This is a perfectly valid outcome though, so don't spam the log!
-// // 		UE_LOG(LogEQS, Warning,
-// // 			TEXT("Unable to generate any center location candidates at all, Generate Items has failed.  %s in %s."),
-// // 			*GetName(), *QueryInstance.QueryName);
-// 	}
-}
+	// gather context with raw data, so it can be used by derived generators
+	FEnvQueryContextData ContextData;
+	const bool bSuccess = QueryInstance.PrepareContext(CircleCenter, ContextData);
 
-void UEnvQueryGenerator_OnCircle::PrepareGeneratorContext(FEnvQueryInstance& QueryInstance,
-														  TSharedPtr<FVectorAndDataContainer>& CenterLocationCandidates) const
-{
-	FVectorAndDataContainer* CandidateData = CenterLocationCandidates.Get();
-	if (CandidateData == NULL)
-	{	// Make sure we have the data holder!
-		CenterLocationCandidates = MakeShareable(new FVectorAndDataContainer);
-		CandidateData = CenterLocationCandidates.Get();
-	}
-	else // We already had it... so make sure it's clear to start, just in case we need to clear old data.
+	if (bSuccess && ContextData.ValueType && ContextData.ValueType->IsChildOf(UEnvQueryItemType_VectorBase::StaticClass()))
 	{
-		CandidateData->Reset();
-	}
-	check(CandidateData != NULL);
+		UEnvQueryItemType_VectorBase* DefTypeOb = (UEnvQueryItemType_VectorBase*)ContextData.ValueType->GetDefaultObject();
+		const uint16 DefTypeValueSize = DefTypeOb->GetValueSize();
+		uint8* RawData = ContextData.RawData.GetData();
 
-	// Now use the generic vector context preparation.
-	QueryInstance.PrepareContext(CircleCenter, CandidateData->GetLocationArray());
-
-	// If there are no entries, make a fake one based on the querier.
-	if (CandidateData->GetNumEntries() < 1)
-	{	
-		UE_LOG(LogEQS, Warning, TEXT("No context locations found, using querier as context for %s in instance %s"), *GetName(), *QueryInstance.QueryName);
-		AActor* Querier = Cast<AActor>(QueryInstance.Owner.Get());
-		if (Querier != NULL)
+		for (int32 ValueIndex = 0; ValueIndex < ContextData.NumValues; ValueIndex++)
 		{
-			FVector CenterLocation = Querier->GetActorLocation();
-			CandidateData->GetLocationArray().Add(CenterLocation);
+			const FVector ContextItemLocation = DefTypeOb->GetItemLocation(RawData);
+			GenerateItemsForCircle(RawData, DefTypeOb, ContextItemLocation, StartDirection, StepsCount, AngleStep, QueryInstance);
+
+			RawData += DefTypeValueSize;
 		}
 	}
 }
 
-void UEnvQueryGenerator_OnCircle::GenerateItemsForCircle(TSharedPtr<FVectorAndDataContainer>& CenterLocationCandidates,
-	int32 CenterLocationCandidateIndex,
-	const FVector& CenterLocation, const FVector& StartDirection,
+void UEnvQueryGenerator_OnCircle::GenerateItemsForCircle(uint8* ContextRawData, UEnvQueryItemType* ContextItemType,
+	const FVector& CenterLocation, const FVector& StartDirection, 
 	int32 StepsCount, float AngleStep, FEnvQueryInstance& OutQueryInstance) const
 {
 	TArray<FVector> ItemCandidates;
@@ -360,17 +336,15 @@ void UEnvQueryGenerator_OnCircle::GenerateItemsForCircle(TSharedPtr<FVectorAndDa
 	}
 #endif
 
-	AddItemDataForCircle(ItemCandidates, CenterLocationCandidates, CenterLocationCandidateIndex, OutQueryInstance);
+	AddItemDataForCircle(ContextRawData, ContextItemType, ItemCandidates, OutQueryInstance);
 }
 
-void UEnvQueryGenerator_OnCircle::AddItemDataForCircle(const TArray<FVector>& ItemCandidates,
-													   TSharedPtr<FVectorAndDataContainer>& CenterLocationCandidates,
-													   int32 CenterLocationCandidateIndex,
-													   FEnvQueryInstance& OutQueryInstance) const
+void UEnvQueryGenerator_OnCircle::AddItemDataForCircle(uint8* ContextRawData, UEnvQueryItemType* ContextItemType,
+	const TArray<FVector>& Locations, FEnvQueryInstance& OutQueryInstance) const
 {
-	for (int32 Step = 0; Step < ItemCandidates.Num(); ++Step)
+	for (int32 Idx = 0; Idx < Locations.Num(); Idx++)
 	{
-		OutQueryInstance.AddItemData<UEnvQueryItemType_Point>(ItemCandidates[Step]);
+		OutQueryInstance.AddItemData<UEnvQueryItemType_Point>(Locations[Idx]);
 	}
 }
 
@@ -386,8 +360,8 @@ FText UEnvQueryGenerator_OnCircle::GetDescriptionTitle() const
 FText UEnvQueryGenerator_OnCircle::GetDescriptionDetails() const
 {
 	FFormatNamedArguments Args;
-	Args.Add(TEXT("Radius"), FText::FromString(UEnvQueryTypes::DescribeFloatParam(Radius)));
-	Args.Add(TEXT("ItemSpacing"), FText::FromString(UEnvQueryTypes::DescribeFloatParam(ItemSpacing)));
+	Args.Add(TEXT("Radius"), FText::FromString(CircleRadius.ToString()));
+	Args.Add(TEXT("ItemSpacing"), FText::FromString(SpaceBetween.ToString()));
 	Args.Add(TEXT("TraceData"), TraceData.ToText(FEnvTraceData::Detailed));
 
 	FText Desc = FText::Format(LOCTEXT("OnCircleDescription", "radius: {Radius}, item span: {ItemSpacing}\n{TraceData}"), Args);
@@ -417,27 +391,26 @@ FText UEnvQueryGenerator_OnCircle::GetDescriptionDetails() const
 
 void UEnvQueryGenerator_OnCircle::PostEditChangeProperty( FPropertyChangedEvent& PropertyChangedEvent) 
 {
-	static const FName NAME_Angle = GET_MEMBER_NAME_CHECKED(UEnvQueryGenerator_OnCircle, Angle);
-	static const FName NAME_Radius = GET_MEMBER_NAME_CHECKED(UEnvQueryGenerator_OnCircle, Radius);
-
 	if (PropertyChangedEvent.Property != NULL)
 	{
 		const FName PropName = PropertyChangedEvent.MemberProperty->GetFName();
-		if (PropName == NAME_Angle)
+		if (PropName == GET_MEMBER_NAME_CHECKED(UEnvQueryGenerator_OnCircle, ArcAngle))
 		{
-			Angle.Value = FMath::Clamp(Angle.Value, 0.0f, 360.f);
-			AngleRadians = FMath::DegreesToRadians(Angle.Value);
-			bDefineArc = Angle.Value < 360.f && Angle.Value > 0.f;
+			ArcAngle.DefaultValue = FMath::Clamp(ArcAngle.DefaultValue, 0.0f, 360.f);
+			AngleRadians = FMath::DegreesToRadians(ArcAngle.DefaultValue);
+			bDefineArc = (ArcAngle.DefaultValue < 360.f) && (ArcAngle.DefaultValue > 0.f);
 		}
-		else if (PropName == NAME_Radius)
+		else if (PropName == GET_MEMBER_NAME_CHECKED(UEnvQueryGenerator_OnCircle, CircleRadius))
 		{
-			if (Radius.Value <= 0.f)
+			if (CircleRadius.DefaultValue <= 0.f)
 			{
-				Radius.Value = 100.f;
+				CircleRadius.DefaultValue = 100.f;
 			}
 		}
-
-		ItemSpacing.Value = FMath::Max(1.0f, ItemSpacing.Value);
+		else if (PropName == GET_MEMBER_NAME_CHECKED(UEnvQueryGenerator_OnCircle, SpaceBetween))
+		{
+			SpaceBetween.DefaultValue = FMath::Max(1.0f, SpaceBetween.DefaultValue);
+		}
 	}
 
 	Super::PostEditChangeProperty(PropertyChangedEvent);
diff --git a/Engine/Source/Runtime/AIModule/Private/EnvironmentQuery/Generators/EnvQueryGenerator_PathingGrid.cpp b/Engine/Source/Runtime/AIModule/Private/EnvironmentQuery/Generators/EnvQueryGenerator_PathingGrid.cpp
--- a/Engine/Source/Runtime/AIModule/Private/EnvironmentQuery/Generators/EnvQueryGenerator_PathingGrid.cpp
+++ b/Engine/Source/Runtime/AIModule/Private/EnvironmentQuery/Generators/EnvQueryGenerator_PathingGrid.cpp
@@ -13,11 +13,28 @@ UEnvQueryGenerator_PathingGrid::UEnvQueryGenerator_PathingGrid(const FObjectInit
 {
 	GenerateAround = UEnvQueryContext_Querier::StaticClass();
 	ItemType = UEnvQueryItemType_Point::StaticClass();
+	MaxDistance.DefaultValue = 100.0f;
+	SpaceBetween.DefaultValue = 10.0f;
+	PathToItem.DefaultValue = true;
+
+	// keep deprecated properties initialized
 	MaxPathDistance.Value = 100.0f;
 	Density.Value = 10.0f;
 	PathFromContext.Value = true;
 }
 
+void UEnvQueryGenerator_PathingGrid::PostLoad()
+{
+	if (VerNum < EnvQueryGeneratorVersion::DataProviders)
+	{
+		MaxPathDistance.Convert(this, MaxDistance);
+		Density.Convert(this, SpaceBetween);
+		PathFromContext.Convert(this, PathToItem);
+	}
+
+	Super::PostLoad();
+}
+
 void UEnvQueryGenerator_PathingGrid::GenerateItems(FEnvQueryInstance& QueryInstance) const
 {
 #if WITH_RECAST
@@ -27,15 +44,14 @@ void UEnvQueryGenerator_PathingGrid::GenerateItems(FEnvQueryInstance& QueryInsta
 		return;
 	}
 
-	float PathDistanceValue = 0.0f;
-	float DensityValue = 0.0f;
-	bool bFromContextValue = true;
-	if (!QueryInstance.GetParamValue(MaxPathDistance, PathDistanceValue, TEXT("MaxPathDistance")) ||
-		!QueryInstance.GetParamValue(Density, DensityValue, TEXT("Density")) ||
-		!QueryInstance.GetParamValue(PathFromContext, bFromContextValue, TEXT("PathFromContext")))
-	{
-		return;
-	}
+	UObject* BindOwner = QueryInstance.Owner.Get();
+	MaxDistance.BindData(BindOwner, QueryInstance.QueryID);
+	SpaceBetween.BindData(BindOwner, QueryInstance.QueryID);
+	PathToItem.BindData(BindOwner, QueryInstance.QueryID);
+
+	float PathDistanceValue = MaxDistance.GetValue();
+	float DensityValue = SpaceBetween.GetValue();
+	bool bFromContextValue = PathToItem.GetValue();
 
 	const int32 ItemCount = FPlatformMath::TruncToInt((PathDistanceValue * 2.0f / DensityValue) + 1);
 	const int32 ItemCountHalf = ItemCount / 2;
@@ -88,21 +104,14 @@ void UEnvQueryGenerator_PathingGrid::GenerateItems(FEnvQueryInstance& QueryInsta
 
 FText UEnvQueryGenerator_PathingGrid::GetDescriptionTitle() const
 {
-	FFormatNamedArguments Args;
-	Args.Add(TEXT("DescriptionTitle"), Super::GetDescriptionTitle());
-	Args.Add(TEXT("DescribeContext"), UEnvQueryTypes::DescribeContext(GenerateAround));
-
-	return FText::Format(LOCTEXT("DescriptionGenerateAroundContext", "{DescriptionTitle}: generate around {DescribeContext}"), Args);
+	return FText::Format(LOCTEXT("DescriptionGenerateAroundContext", "{0}: generate around {1}"),
+		Super::GetDescriptionTitle(), UEnvQueryTypes::DescribeContext(GenerateAround));
 };
 
 FText UEnvQueryGenerator_PathingGrid::GetDescriptionDetails() const
 {
-	FFormatNamedArguments Args;
-	Args.Add(TEXT("MaxPathDistance"), FText::FromString(UEnvQueryTypes::DescribeFloatParam(MaxPathDistance)));
-	Args.Add(TEXT("Density"), FText::FromString(UEnvQueryTypes::DescribeFloatParam(Density)));
-	Args.Add(TEXT("PathFromContext"), FText::FromString(UEnvQueryTypes::DescribeBoolParam(PathFromContext)));
-
-	return FText::Format(LOCTEXT("DescriptionDetailsPathingGrid", "max distance: {MaxPathDistance}, density: {Density}, path from context: {PathFromContext}"), Args);
+	return FText::Format(LOCTEXT("DescriptionDetailsPathingGrid", "max distance: {0}, space between: {1}, path to item: {2}"),
+		FText::FromString(MaxDistance.ToString()), FText::FromString(SpaceBetween.ToString()), FText::FromString(PathToItem.ToString()));
 }
 
 #if WITH_RECAST
diff --git a/Engine/Source/Runtime/AIModule/Private/EnvironmentQuery/Generators/EnvQueryGenerator_SimpleGrid.cpp b/Engine/Source/Runtime/AIModule/Private/EnvironmentQuery/Generators/EnvQueryGenerator_SimpleGrid.cpp
--- a/Engine/Source/Runtime/AIModule/Private/EnvironmentQuery/Generators/EnvQueryGenerator_SimpleGrid.cpp
+++ b/Engine/Source/Runtime/AIModule/Private/EnvironmentQuery/Generators/EnvQueryGenerator_SimpleGrid.cpp
@@ -12,21 +12,34 @@ UEnvQueryGenerator_SimpleGrid::UEnvQueryGenerator_SimpleGrid(const FObjectInitia
 {
 	GenerateAround = UEnvQueryContext_Querier::StaticClass();
 	ItemType = UEnvQueryItemType_Point::StaticClass();
+	GridSize.DefaultValue = 500.0f;
+	SpaceBetween.DefaultValue = 100.0f;
+
+	// keep deprecated properties initialized
 	Radius.Value = 500.0f;
-	Density.Value = 100.0f;
+	Density.Value = 500.0f;
 }
 
-void UEnvQueryGenerator_SimpleGrid::GenerateItems(FEnvQueryInstance& QueryInstance) const
+void UEnvQueryGenerator_SimpleGrid::PostLoad()
 {
-	float RadiusValue = 0.0f;
-	float DensityValue = 0.0f;
-
-	if (!QueryInstance.GetParamValue(Radius, RadiusValue, TEXT("Radius")) ||
-		!QueryInstance.GetParamValue(Density, DensityValue, TEXT("Density")) )
+	if (VerNum < EnvQueryGeneratorVersion::DataProviders)
 	{
-		return;
+		Radius.Convert(this, GridSize);
+		Density.Convert(this, SpaceBetween);
 	}
 
+	Super::PostLoad();
+}
+
+void UEnvQueryGenerator_SimpleGrid::GenerateItems(FEnvQueryInstance& QueryInstance) const
+{
+	UObject* BindOwner = QueryInstance.Owner.Get();
+	GridSize.BindData(BindOwner, QueryInstance.QueryID);
+	SpaceBetween.BindData(BindOwner, QueryInstance.QueryID);
+
+	float RadiusValue = GridSize.GetValue();
+	float DensityValue = SpaceBetween.GetValue();
+
 	const bool bProjectToNavigation = (ProjectionData.TraceMode == EEnvQueryTrace::Navigation);
 #if WITH_RECAST
 	const ARecastNavMesh* NavMesh = bProjectToNavigation ? FEQSHelpers::FindNavMeshForQuery(QueryInstance) : NULL;
@@ -82,20 +95,14 @@ void UEnvQueryGenerator_SimpleGrid::GenerateItems(FEnvQueryInstance& QueryInstan
 
 FText UEnvQueryGenerator_SimpleGrid::GetDescriptionTitle() const
 {
-	FFormatNamedArguments Args;
-	Args.Add(TEXT("DescriptionTitle"), Super::GetDescriptionTitle());
-	Args.Add(TEXT("DescribeContext"), UEnvQueryTypes::DescribeContext(GenerateAround));
-
-	return FText::Format(LOCTEXT("DescriptionGenerateSimpleGridAroundContext", "{DescriptionTitle}: generate around {DescribeContext}"), Args);
+	return FText::Format(LOCTEXT("DescriptionGenerateAroundContext", "{0}: generate around {1}"),
+		Super::GetDescriptionTitle(), UEnvQueryTypes::DescribeContext(GenerateAround));
 };
 
 FText UEnvQueryGenerator_SimpleGrid::GetDescriptionDetails() const
 {
-	FFormatNamedArguments Args;
-	Args.Add(TEXT("Radius"), FText::FromString(UEnvQueryTypes::DescribeFloatParam(Radius)));
-	Args.Add(TEXT("Density"), FText::FromString(UEnvQueryTypes::DescribeFloatParam(Density)));
-
-	FText Desc = FText::Format(LOCTEXT("SimpleGridDescription", "radius: {Radius}, density: {Density}"), Args);
+	FText Desc = FText::Format(LOCTEXT("SimpleGridDescription", "size: {0}, space between: {1}"),
+		FText::FromString(GridSize.ToString()), FText::FromString(SpaceBetween.ToString()));
 
 	FText ProjDesc = ProjectionData.ToText(FEnvTraceData::Brief);
 	if (!ProjDesc.IsEmpty())
diff --git a/Engine/Source/Runtime/AIModule/Private/EnvironmentQuery/Tests/EnvQueryTest_Distance.cpp b/Engine/Source/Runtime/AIModule/Private/EnvironmentQuery/Tests/EnvQueryTest_Distance.cpp
--- a/Engine/Source/Runtime/AIModule/Private/EnvironmentQuery/Tests/EnvQueryTest_Distance.cpp
+++ b/Engine/Source/Runtime/AIModule/Private/EnvironmentQuery/Tests/EnvQueryTest_Distance.cpp
@@ -32,23 +32,11 @@ UEnvQueryTest_Distance::UEnvQueryTest_Distance(const FObjectInitializer& ObjectI
 
 void UEnvQueryTest_Distance::RunTest(FEnvQueryInstance& QueryInstance) const
 {
-// 	float ThresholdValue = 0.0f;
-// 	if (!QueryInstance.GetParamValue(FloatFilter, ThresholdValue, TEXT("FloatFilter")))
-// 	{
-// 		return;
-// 	}
+	FloatValueMin.BindData(QueryInstance.Owner.Get(), QueryInstance.QueryID);
+	float MinThresholdValue = FloatValueMin.GetValue();
 
-	float MinThresholdValue = 0.0f;
-	if (!QueryInstance.GetParamValue(FloatFilterMin, MinThresholdValue, TEXT("FloatFilterMin")))
-	{
-		return;
-	}
-
-	float MaxThresholdValue = 0.0f;
-	if (!QueryInstance.GetParamValue(FloatFilterMax, MaxThresholdValue, TEXT("FloatFilterMax")))
-	{
-		return;
-	}
+	FloatValueMax.BindData(QueryInstance.Owner.Get(), QueryInstance.QueryID);
+	float MaxThresholdValue = FloatValueMax.GetValue();
 
 	// don't support context Item here, it doesn't make any sense
 	TArray<FVector> ContextLocations;
diff --git a/Engine/Source/Runtime/AIModule/Private/EnvironmentQuery/Tests/EnvQueryTest_Dot.cpp b/Engine/Source/Runtime/AIModule/Private/EnvironmentQuery/Tests/EnvQueryTest_Dot.cpp
--- a/Engine/Source/Runtime/AIModule/Private/EnvironmentQuery/Tests/EnvQueryTest_Dot.cpp
+++ b/Engine/Source/Runtime/AIModule/Private/EnvironmentQuery/Tests/EnvQueryTest_Dot.cpp
@@ -22,17 +22,11 @@ UEnvQueryTest_Dot::UEnvQueryTest_Dot(const FObjectInitializer& ObjectInitializer
 
 void UEnvQueryTest_Dot::RunTest(FEnvQueryInstance& QueryInstance) const
 {
-	float MinThresholdValue = 0.0f;
-	if (!QueryInstance.GetParamValue(FloatFilterMin, MinThresholdValue, TEXT("FloatFilterMin")))
-	{
-		return;
-	}
+	FloatValueMin.BindData(QueryInstance.Owner.Get(), QueryInstance.QueryID);
+	float MinThresholdValue = FloatValueMin.GetValue();
 
-	float MaxThresholdValue = 0.0f;
-	if (!QueryInstance.GetParamValue(FloatFilterMax, MaxThresholdValue, TEXT("FloatFilterMax")))
-	{
-		return;
-	}
+	FloatValueMax.BindData(QueryInstance.Owner.Get(), QueryInstance.QueryID);
+	float MaxThresholdValue = FloatValueMax.GetValue();
 
 	// gather all possible directions: for contexts different than Item
 	TArray<FVector> LineADirs;
diff --git a/Engine/Source/Runtime/AIModule/Private/EnvironmentQuery/Tests/EnvQueryTest_GameplayTags.cpp b/Engine/Source/Runtime/AIModule/Private/EnvironmentQuery/Tests/EnvQueryTest_GameplayTags.cpp
--- a/Engine/Source/Runtime/AIModule/Private/EnvironmentQuery/Tests/EnvQueryTest_GameplayTags.cpp
+++ b/Engine/Source/Runtime/AIModule/Private/EnvironmentQuery/Tests/EnvQueryTest_GameplayTags.cpp
@@ -16,7 +16,8 @@ UEnvQueryTest_GameplayTags::UEnvQueryTest_GameplayTags(const FObjectInitializer&
 }
 
 bool UEnvQueryTest_GameplayTags::SatisfiesTest(IGameplayTagAssetInterface* ItemGameplayTagAssetInterface) const
-{	// Currently we're requiring that the test only be run on items that implement the interface.  In theory, we could
+{
+	// Currently we're requiring that the test only be run on items that implement the interface.  In theory, we could
 	// (instead of checking) support correctly passing or failing on items that don't implement the interface or
 	// just have a NULL item by testing them as though they have the interface with no gameplay tags.  However, that
 	// seems potentially confusing, since certain tests could actually pass.
@@ -51,11 +52,8 @@ bool UEnvQueryTest_GameplayTags::SatisfiesTest(IGameplayTagAssetInterface* ItemG
 
 void UEnvQueryTest_GameplayTags::RunTest(FEnvQueryInstance& QueryInstance) const
 {
-	bool bWantsValid = true;
-	if (!QueryInstance.GetParamValue(BoolFilter, bWantsValid, TEXT("BoolFilter")))
-	{
-		return;
-	}
+	BoolValue.BindData(QueryInstance.Owner.Get(), QueryInstance.QueryID);
+	bool bWantsValid = BoolValue.GetValue();
 
 	// loop through all items
 	for (FEnvQueryInstance::ItemIterator It(this, QueryInstance); It; ++It)
@@ -84,5 +82,5 @@ FString UEnvQueryTest_GameplayTags::GetDescriptionTitle() const
 
 FText UEnvQueryTest_GameplayTags::GetDescriptionDetails() const
 {
-	return GameplayTags.ToMatchingText(TagsToMatch, !BoolFilter.Value);
+	return GameplayTags.ToMatchingText(TagsToMatch, !BoolValue.DefaultValue);
 }
diff --git a/Engine/Source/Runtime/AIModule/Private/EnvironmentQuery/Tests/EnvQueryTest_Pathfinding.cpp b/Engine/Source/Runtime/AIModule/Private/EnvironmentQuery/Tests/EnvQueryTest_Pathfinding.cpp
--- a/Engine/Source/Runtime/AIModule/Private/EnvironmentQuery/Tests/EnvQueryTest_Pathfinding.cpp
+++ b/Engine/Source/Runtime/AIModule/Private/EnvironmentQuery/Tests/EnvQueryTest_Pathfinding.cpp
@@ -15,34 +15,36 @@ UEnvQueryTest_Pathfinding::UEnvQueryTest_Pathfinding(const FObjectInitializer& O
 	Cost = EEnvTestCost::High;
 	ValidItemType = UEnvQueryItemType_VectorBase::StaticClass();
 	TestMode = EEnvTestPathfinding::PathExist;
+	PathFromContext.DefaultValue = true;
+	SkipUnreachable.DefaultValue = true;
+	UseHierarchicalPathfinding.DefaultValue = true;
+	FloatValueMin.DefaultValue = 1000.0f;
+	FloatValueMax.DefaultValue = 1000.0f;
+
+	// keep deprecated properties initialized
 	PathToItem.Value = true;
-	FloatFilter.Value = 1000.0f;
 	DiscardUnreachable.Value = true;
 	HierarchicalPathfinding.Value = true;
+	FloatFilterMin.Value = 1000.0f;
+	FloatFilterMax.Value = 1000.0f;
 }
 
 void UEnvQueryTest_Pathfinding::RunTest(FEnvQueryInstance& QueryInstance) const
 {
-	bool bWantsPath = false;
-	bool bPathToItem = false;
-	bool bHierarchical = false;
-	bool bDiscardFailed = false;
-// 	float UseThreshold = 0.0f;
-	float MinThresholdValue = 0.0f;
-	float MaxThresholdValue = 0.0f;
-
-	if (!QueryInstance.GetParamValue(BoolFilter, bWantsPath, TEXT("BoolFilter")) ||
-		!QueryInstance.GetParamValue(PathToItem, bPathToItem, TEXT("PathToItem")) ||
-		!QueryInstance.GetParamValue(HierarchicalPathfinding, bHierarchical, TEXT("HierarchicalPathfinding")) ||
-		!QueryInstance.GetParamValue(DiscardUnreachable, bDiscardFailed, TEXT("DiscardUnreachable")) ||
-// 		!QueryInstance.GetParamValue(FloatFilter, UseThreshold, TEXT("FloatFilter")) ||
-		!QueryInstance.GetParamValue(FloatFilterMax, MaxThresholdValue, TEXT("FloatFilterMax")) ||
-		!QueryInstance.GetParamValue(FloatFilterMin, MinThresholdValue, TEXT("FloatFilterMin"))
-	   )
-
-	{
-		return;
-	}
+	UObject* DataOwner = QueryInstance.Owner.Get();
+	BoolValue.BindData(DataOwner, QueryInstance.QueryID);
+	PathFromContext.BindData(DataOwner, QueryInstance.QueryID);
+	SkipUnreachable.BindData(DataOwner, QueryInstance.QueryID);
+	UseHierarchicalPathfinding.BindData(DataOwner, QueryInstance.QueryID);
+	FloatValueMin.BindData(DataOwner, QueryInstance.QueryID);
+	FloatValueMax.BindData(DataOwner, QueryInstance.QueryID);
+
+	bool bWantsPath = BoolValue.GetValue();
+	bool bPathToItem = PathFromContext.GetValue();
+	bool bHierarchical = UseHierarchicalPathfinding.GetValue();
+	bool bDiscardFailed = SkipUnreachable.GetValue();
+	float MinThresholdValue = FloatValueMin.GetValue();
+	float MaxThresholdValue = FloatValueMax.GetValue();
 
 	UNavigationSystem* NavSys = QueryInstance.World->GetNavigationSystem();
 	ANavigationData* NavData = FindNavigationData(NavSys, QueryInstance.Owner.Get());
@@ -118,71 +120,52 @@ FString UEnvQueryTest_Pathfinding::GetDescriptionTitle() const
 {
 	FString ModeDesc[] = { TEXT("PathExist"), TEXT("PathCost"), TEXT("PathLength") };
 
-	FString DirectionDesc = PathToItem.IsNamedParam() ?
-		FString::Printf(TEXT("%s, direction: %s"), *UEnvQueryTypes::DescribeContext(Context).ToString(), *UEnvQueryTypes::DescribeBoolParam(PathToItem)) :
-		FString::Printf(TEXT("%s %s"), PathToItem.Value ? TEXT("from") : TEXT("to"), *UEnvQueryTypes::DescribeContext(Context).ToString());
+	FString DirectionDesc = PathFromContext.IsDynamic() ?
+		FString::Printf(TEXT("%s, direction: %s"), *UEnvQueryTypes::DescribeContext(Context).ToString(), *PathFromContext.ToString()) :
+		FString::Printf(TEXT("%s %s"), PathFromContext.DefaultValue ? TEXT("from") : TEXT("to"), *UEnvQueryTypes::DescribeContext(Context).ToString());
 
 	return FString::Printf(TEXT("%s: %s"), *ModeDesc[TestMode], *DirectionDesc);
 }
 
 FText UEnvQueryTest_Pathfinding::GetDescriptionDetails() const
 {
-	FText AdditionalDesc;
-	if (HierarchicalPathfinding.IsNamedParam())
+	FText HPathDesc = LOCTEXT("HierarchicalPathfinding", "hierarchical pathfinding");
+	FText Desc1;
+	if (UseHierarchicalPathfinding.IsDynamic())
 	{
-		FFormatNamedArguments Args;
-		Args.Add(TEXT("PathName"), FText::FromString(HierarchicalPathfinding.ParamName.ToString()));
-		AdditionalDesc = FText::Format(LOCTEXT("HierarchicalPathfindingWithName", "Hierarchical pathfinding: {PathName}"), Args);
+		Desc1 = FText::Format(FText::FromString("{0}: {1}"), HPathDesc, FText::FromString(UseHierarchicalPathfinding.ToString()));
 	}
-	else if (HierarchicalPathfinding.Value)
+	else if (UseHierarchicalPathfinding.DefaultValue)
 	{
-		AdditionalDesc = LOCTEXT("UseHierarchicalPathfinding", "Use hierarchical pathfinding");
+		Desc1 = HPathDesc;
 	}
 
-	if (DiscardUnreachable.IsNamedParam())
+	FText DiscardDesc = LOCTEXT("DiscardUnreachable", "discard unreachable");
+	FText Desc2;
+	if (SkipUnreachable.IsDynamic())
 	{
-		FFormatNamedArguments Args;
-		Args.Add(TEXT("Description"), AdditionalDesc);
-		Args.Add(TEXT("UnreachableParamName"), FText::FromName(DiscardUnreachable.ParamName));
-
-		if(AdditionalDesc.IsEmpty())
-		{
-			AdditionalDesc = FText::Format(LOCTEXT("DiscardUnreachableWithParam", "Discard unreachable: {UnreachableParamName}"), Args);
-		}
-		else
-		{
-			AdditionalDesc = FText::Format(LOCTEXT("DescWithUnreachableAndParam", "{Description}, discard unreachable: {UnreachableParamName}"), Args);
-		}
+		Desc2 = FText::Format(FText::FromString("{0}: {1}"), DiscardDesc, FText::FromString(SkipUnreachable.ToString()));
 	}
-	else if (DiscardUnreachable.Value)
+	else if (SkipUnreachable.DefaultValue)
 	{
-		if(AdditionalDesc.IsEmpty())
-		{
-			AdditionalDesc = LOCTEXT("DiscardUnreachable", "Discard unreachable");
-		}
-		else
-		{
-			FFormatNamedArguments Args;
-			Args.Add(TEXT("Description"), AdditionalDesc);
-			AdditionalDesc = FText::Format(LOCTEXT("DescWithUnreachable", "{Description}, discard unreachable"), Args);
-		}
+		Desc2 = DiscardDesc;
 	}
 
-	FFormatNamedArguments Args;
-	Args.Add(TEXT("Description"), AdditionalDesc);
-
-	if (GetWorkOnFloatValues())
+	FText TestParamDesc = GetWorkOnFloatValues() ? DescribeFloatTestParams() : DescribeBoolTestParams("existing path");
+	if (!Desc1.IsEmpty() && !Desc2.IsEmpty())
 	{
-		Args.Add(TEXT("FloatParamsDescription"), DescribeFloatTestParams());
-		AdditionalDesc = FText::Format(LOCTEXT("DescriptionWithDescribedFloatParams", "{Description}\n{FloatParamsDescription}"), Args);
+		return FText::Format(FText::FromString("{0}, {1}\n{2}"), Desc1, Desc2, TestParamDesc);
 	}
-	else
+	else if (!Desc1.IsEmpty())
+	{
+		return FText::Format(FText::FromString("{0}\n{1}"), Desc1, TestParamDesc);
+	}
+	else if (!Desc2.IsEmpty())
 	{
-		Args.Add(TEXT("BoolParamsDescription"), DescribeBoolTestParams("existing path"));
-		AdditionalDesc = FText::Format(LOCTEXT("DescriptionWithDescribedBoolParams", "{Description}\n{BoolParamsDescription}"), Args);
+		return FText::Format(FText::FromString("{0}\n{1}"), Desc2, TestParamDesc);
 	}
 
-	return AdditionalDesc;
+	return TestParamDesc;
 }
 
 #if WITH_EDITOR
@@ -192,13 +175,19 @@ void UEnvQueryTest_Pathfinding::PostEditChangeProperty(FPropertyChangedEvent& Pr
 	if (PropertyChangedEvent.Property && PropertyChangedEvent.Property->GetFName() == GET_MEMBER_NAME_CHECKED(UEnvQueryTest_Pathfinding,TestMode))
 	{
 		SetWorkOnFloatValues(TestMode != EEnvTestPathfinding::PathExist);
-		Condition = EEnvTestCondition::NoCondition;
 	}
 }
 #endif
 
 void UEnvQueryTest_Pathfinding::PostLoad()
 {
+	if (VerNum < EnvQueryTestVersion::DataProviders)
+	{
+		PathToItem.Convert(this, PathFromContext);
+		DiscardUnreachable.Convert(this, SkipUnreachable);
+		HierarchicalPathfinding.Convert(this, UseHierarchicalPathfinding);
+	}
+
 	Super::PostLoad();
 	
 	SetWorkOnFloatValues(TestMode != EEnvTestPathfinding::PathExist);
diff --git a/Engine/Source/Runtime/AIModule/Private/EnvironmentQuery/Tests/EnvQueryTest_Random.cpp b/Engine/Source/Runtime/AIModule/Private/EnvironmentQuery/Tests/EnvQueryTest_Random.cpp
--- a/Engine/Source/Runtime/AIModule/Private/EnvironmentQuery/Tests/EnvQueryTest_Random.cpp
+++ b/Engine/Source/Runtime/AIModule/Private/EnvironmentQuery/Tests/EnvQueryTest_Random.cpp
@@ -14,17 +14,11 @@ UEnvQueryTest_Random::UEnvQueryTest_Random(const FObjectInitializer& ObjectIniti
 
 void UEnvQueryTest_Random::RunTest(FEnvQueryInstance& QueryInstance) const
 {
-	float MinThresholdValue = 0.0f;
-	if (!QueryInstance.GetParamValue(FloatFilterMin, MinThresholdValue, TEXT("FloatFilterMin")))
-	{
-		return;
-	}
+	FloatValueMin.BindData(QueryInstance.Owner.Get(), QueryInstance.QueryID);
+	float MinThresholdValue = FloatValueMin.GetValue();
 
-	float MaxThresholdValue = 0.0f;
-	if (!QueryInstance.GetParamValue(FloatFilterMax, MaxThresholdValue, TEXT("FloatFilterMax")))
-	{
-		return;
-	}
+	FloatValueMax.BindData(QueryInstance.Owner.Get(), QueryInstance.QueryID);
+	float MaxThresholdValue = FloatValueMax.GetValue();
 
 	// loop through all items
 	for (FEnvQueryInstance::ItemIterator It(this, QueryInstance); It; ++It)
diff --git a/Engine/Source/Runtime/AIModule/Private/EnvironmentQuery/Tests/EnvQueryTest_Trace.cpp b/Engine/Source/Runtime/AIModule/Private/EnvironmentQuery/Tests/EnvQueryTest_Trace.cpp
--- a/Engine/Source/Runtime/AIModule/Private/EnvironmentQuery/Tests/EnvQueryTest_Trace.cpp
+++ b/Engine/Source/Runtime/AIModule/Private/EnvironmentQuery/Tests/EnvQueryTest_Trace.cpp
@@ -14,28 +14,21 @@ UEnvQueryTest_Trace::UEnvQueryTest_Trace(const FObjectInitializer& ObjectInitial
 	SetWorkOnFloatValues(false);
 	
 	Context = UEnvQueryContext_Querier::StaticClass();
-	TraceToItem.Value = false;
-	ItemOffsetZ.Value = 0.0f;
-	ContextOffsetZ.Value = 0.0f;
-
 	TraceData.SetGeometryOnly();
 }
 
 void UEnvQueryTest_Trace::RunTest(FEnvQueryInstance& QueryInstance) const
 {
-	bool bWantsHit = false;
-	bool bTraceToItem = false;
-	float ItemZ = 0.0f;
-	float ContextZ = 0.0f;
-
-	if (!QueryInstance.GetParamValue(BoolFilter, bWantsHit, TEXT("BoolFilter")) ||
-		!QueryInstance.GetParamValue(TraceToItem, bTraceToItem, TEXT("TraceToItem")) ||
-		!QueryInstance.GetParamValue(ItemOffsetZ, ItemZ, TEXT("ItemOffsetZ")) ||
-		!QueryInstance.GetParamValue(ContextOffsetZ, ContextZ, TEXT("ContextOffsetZ"))
-		)
-	{
-		return;
-	}
+	UObject* DataOwner = QueryInstance.Owner.Get();
+	BoolValue.BindData(DataOwner, QueryInstance.QueryID);
+	TraceFromContext.BindData(DataOwner, QueryInstance.QueryID);
+	ItemHeightOffset.BindData(DataOwner, QueryInstance.QueryID);
+	ContextHeightOffset.BindData(DataOwner, QueryInstance.QueryID);
+
+	bool bWantsHit = BoolValue.GetValue();
+	bool bTraceToItem = TraceFromContext.GetValue();
+	float ItemZ = ItemHeightOffset.GetValue();
+	float ContextZ = ContextHeightOffset.GetValue();
 
 	TArray<FVector> ContextLocations;
 	if (!QueryInstance.PrepareContext(Context, ContextLocations))
@@ -100,9 +93,9 @@ FString UEnvQueryTest_Trace::GetDescriptionTitle() const
 	UEnum* ChannelEnum = FindObject<UEnum>(ANY_PACKAGE, TEXT("ETraceTypeQuery"), true);
 	FString ChannelDesc = ChannelEnum->GetEnumText(TraceData.TraceChannel).ToString();
 
-	FString DirectionDesc = TraceToItem.IsNamedParam() ?
-		FString::Printf(TEXT("%s, direction: %s"), *UEnvQueryTypes::DescribeContext(Context).ToString(), *UEnvQueryTypes::DescribeBoolParam(TraceToItem)) :
-		FString::Printf(TEXT("%s %s"), TraceToItem.Value ? TEXT("from") : TEXT("to"), *UEnvQueryTypes::DescribeContext(Context).ToString());
+	FString DirectionDesc = TraceFromContext.IsDynamic() ?
+		FString::Printf(TEXT("%s, direction: %s"), *UEnvQueryTypes::DescribeContext(Context).ToString(), *TraceFromContext.ToString()) :
+		FString::Printf(TEXT("%s %s"), TraceFromContext.DefaultValue ? TEXT("from") : TEXT("to"), *UEnvQueryTypes::DescribeContext(Context).ToString());
 
 	return FString::Printf(TEXT("%s: %s on %s"), 
 		*Super::GetDescriptionTitle(), *DirectionDesc, *ChannelDesc);
@@ -110,11 +103,8 @@ FString UEnvQueryTest_Trace::GetDescriptionTitle() const
 
 FText UEnvQueryTest_Trace::GetDescriptionDetails() const
 {
-	FFormatNamedArguments Args;
-	Args.Add(TEXT("TraceData"),  TraceData.ToText(FEnvTraceData::Detailed));
-	Args.Add(TEXT("TestParams"),  DescribeBoolTestParams("hit"));
-	
-	return LOCTEXT("TraceDescription", "{TraceData}\n{TestParams}");
+	return FText::Format(FText::FromString("{0}\n{1}"),
+		TraceData.ToText(FEnvTraceData::Detailed), DescribeBoolTestParams("hit"));
 }
 
 bool UEnvQueryTest_Trace::RunLineTraceTo(const FVector& ItemPos, const FVector& ContextPos, AActor* ItemActor, UWorld* World, enum ECollisionChannel Channel, const FCollisionQueryParams& Params, const FVector& Extent)
@@ -189,4 +179,16 @@ bool UEnvQueryTest_Trace::RunCapsuleTraceFrom(const FVector& ItemPos, const FVec
 	return bHit;
 }
 
+void UEnvQueryTest_Trace::PostLoad()
+{
+	if (VerNum < EnvQueryTestVersion::DataProviders)
+	{
+		TraceToItem.Convert(this, TraceFromContext);
+		ItemOffsetZ.Convert(this, ItemHeightOffset);
+		ContextOffsetZ.Convert(this, ContextHeightOffset);
+	}
+
+	Super::PostLoad();
+}
+
 #undef LOCTEXT_NAMESPACE
