Fixed code relying on SLATE_TEXT_ATTRIBUTE for STextBlock.

diff --git a/Engine/Source/Editor/DetailCustomizations/Private/BodyInstanceCustomization.cpp b/Engine/Source/Editor/DetailCustomizations/Private/BodyInstanceCustomization.cpp
--- a/Engine/Source/Editor/DetailCustomizations/Private/BodyInstanceCustomization.cpp
+++ b/Engine/Source/Editor/DetailCustomizations/Private/BodyInstanceCustomization.cpp
@@ -437,7 +437,7 @@ void FBodyInstanceCustomization::CreateCustomCollisionSetup( TSharedRef<class IP
 				.Content()
 				[
 					SNew(STextBlock)
-					.Text(DisplayName)
+					.Text(FText::FromString(DisplayName))
 					.Font( IDetailLayoutBuilder::GetDetailFont() )
 				]
 			]
@@ -532,7 +532,7 @@ void FBodyInstanceCustomization::CreateCustomCollisionSetup( TSharedRef<class IP
 				.Content()
 				[
 					SNew(STextBlock)
-					.Text(DisplayName)
+					.Text(FText::FromString(DisplayName))
 					.Font( IDetailLayoutBuilder::GetDetailFont() )
 				]
 			]
@@ -604,7 +604,7 @@ END_SLATE_FUNCTION_BUILD_OPTIMIZATION
 
 TSharedRef<SWidget> FBodyInstanceCustomization::MakeObjectTypeComboWidget( TSharedPtr<FString> InItem )
 {
-	return SNew(STextBlock) .Text( *InItem ) .Font( IDetailLayoutBuilder::GetDetailFont() );
+	return SNew(STextBlock) .Text( FText::FromString(*InItem) ) .Font( IDetailLayoutBuilder::GetDetailFont() );
 }
 
 void FBodyInstanceCustomization::OnObjectTypeChanged( TSharedPtr<FString> NewSelection, ESelectInfo::Type SelectInfo  )
@@ -627,15 +627,15 @@ void FBodyInstanceCustomization::OnObjectTypeChanged( TSharedPtr<FString> NewSel
 	}
 }
 
-FString FBodyInstanceCustomization::GetObjectTypeComboBoxContent() const
+FText FBodyInstanceCustomization::GetObjectTypeComboBoxContent() const
 {
 	FName ObjectTypeName;
 	if (ObjectTypeHandle->GetValue(ObjectTypeName) == FPropertyAccess::Result::MultipleValues)
 	{
-		return TEXT("Multiple Values");
+		return LOCTEXT("MultipleValues", "Multiple Values");
 	}
 
-	return *ObjectTypeComboBox.Get()->GetSelectedItem().Get();
+	return FText::FromString(*ObjectTypeComboBox.Get()->GetSelectedItem().Get());
 }
 
 TSharedRef<SWidget> FBodyInstanceCustomization::MakeCollisionProfileComboWidget(TSharedPtr<FString> InItem)
@@ -650,7 +650,7 @@ TSharedRef<SWidget> FBodyInstanceCustomization::MakeCollisionProfileComboWidget(
 
 	return
 		SNew(STextBlock)
-		.Text(*InItem)
+		.Text(FText::FromString(*InItem))
 		.ToolTipText(FText::FromString(ProfileMessage))
 		.Font(IDetailLayoutBuilder::GetDetailFont());
 }
@@ -812,18 +812,18 @@ EVisibility FBodyInstanceCustomization::ShouldShowCustomCollisionSetup() const
 	return EVisibility::Visible;//return (bDisplayAdvancedCollisionSettings)? EVisibility::Visible : EVisibility::Hidden;
 }
 
-FString FBodyInstanceCustomization::GetCollisionProfileComboBoxContent() const
+FText FBodyInstanceCustomization::GetCollisionProfileComboBoxContent() const
 {
 	FName ProfileName;
 	if (CollisionProfileNameHandle->GetValue(ProfileName) == FPropertyAccess::Result::MultipleValues)
 	{
-		return TEXT("Multiple Values");
+		return LOCTEXT("MultipleValues", "Multiple Values");
 	}
 
-	return (*GetProfileString(ProfileName).Get());
+	return FText::FromString(*GetProfileString(ProfileName).Get());
 }
 
-FString FBodyInstanceCustomization::GetCollisionProfileComboBoxToolTip() const
+FText FBodyInstanceCustomization::GetCollisionProfileComboBoxToolTip() const
 {
 	FName ProfileName;
 	if (CollisionProfileNameHandle->GetValue(ProfileName) == FPropertyAccess::Result::Success)
@@ -831,11 +831,12 @@ FString FBodyInstanceCustomization::GetCollisionProfileComboBoxToolTip() const
 		FCollisionResponseTemplate ProfileData;
 		if ( CollisionProfile->GetProfileTemplate(ProfileName, ProfileData) )
 		{
-			return ProfileData.HelpMessage;
+			return FText::FromString(ProfileData.HelpMessage);
 		}
+		return FText::GetEmpty();
 	}
 
-	return TEXT("Multiple Values");
+	return LOCTEXT("MultipleValues", "Multiple Values");
 }
 
 void FBodyInstanceCustomization::OnCollisionChannelChanged(ECheckBoxState InNewValue, int32 ValidIndex, ECollisionResponse InCollisionResponse)
diff --git a/Engine/Source/Editor/DetailCustomizations/Private/BodyInstanceCustomization.h b/Engine/Source/Editor/DetailCustomizations/Private/BodyInstanceCustomization.h
--- a/Engine/Source/Editor/DetailCustomizations/Private/BodyInstanceCustomization.h
+++ b/Engine/Source/Editor/DetailCustomizations/Private/BodyInstanceCustomization.h
@@ -30,14 +30,14 @@ private:
 	// Profile combo related
 	TSharedRef<SWidget> MakeCollisionProfileComboWidget( TSharedPtr<FString> InItem );
 	void OnCollisionProfileChanged( TSharedPtr<FString> NewSelection, ESelectInfo::Type SelectInfo, IDetailGroup* CollisionGroup );
-	FString GetCollisionProfileComboBoxContent() const;
-	FString GetCollisionProfileComboBoxToolTip() const;
+	FText GetCollisionProfileComboBoxContent() const;
+	FText GetCollisionProfileComboBoxToolTip() const;
 	void OnCollisionProfileComboOpening();
 
 	// Movement channel related
 	TSharedRef<SWidget> MakeObjectTypeComboWidget( TSharedPtr<FString> InItem );
 	void OnObjectTypeChanged( TSharedPtr<FString> NewSelection, ESelectInfo::Type SelectInfo  );
-	FString GetObjectTypeComboBoxContent() const;
+	FText GetObjectTypeComboBoxContent() const;
 	int32 InitializeObjectTypeComboList();
 
 	// set to default for profile setting
diff --git a/Engine/Source/Editor/DetailCustomizations/Private/BrushDetails.cpp b/Engine/Source/Editor/DetailCustomizations/Private/BrushDetails.cpp
--- a/Engine/Source/Editor/DetailCustomizations/Private/BrushDetails.cpp
+++ b/Engine/Source/Editor/DetailCustomizations/Private/BrushDetails.cpp
@@ -225,8 +225,8 @@ void FBrushDetails::CustomizeDetails( IDetailLayoutBuilder& DetailLayout )
 			.ButtonContent()
 			[
 				SNew(STextBlock)
-				.Text(NSLOCTEXT("BrushDetails", "PolygonsMenu", "Polygons").ToString())
-				.ToolTipText(NSLOCTEXT("BrushDetails", "PolygonsMenu_ToolTip", "Polygon options").ToString())
+				.Text(NSLOCTEXT("BrushDetails", "PolygonsMenu", "Polygons"))
+				.ToolTipText(NSLOCTEXT("BrushDetails", "PolygonsMenu_ToolTip", "Polygon options"))
 				.Font(IDetailLayoutBuilder::GetDetailFont())
 			]
 			.MenuContent()
@@ -244,8 +244,8 @@ void FBrushDetails::CustomizeDetails( IDetailLayoutBuilder& DetailLayout )
 			.ButtonContent()
 			[
 				SNew(STextBlock)
-				.Text(NSLOCTEXT("BrushDetails", "SolidityMenu", "Solidity").ToString())
-				.ToolTipText(NSLOCTEXT("BrushDetails", "SolidityMenu_ToolTip", "Solidity options").ToString())
+				.Text(NSLOCTEXT("BrushDetails", "SolidityMenu", "Solidity"))
+				.ToolTipText(NSLOCTEXT("BrushDetails", "SolidityMenu_ToolTip", "Solidity options"))
 				.Font(IDetailLayoutBuilder::GetDetailFont())
 			]
 			.MenuContent()
@@ -263,8 +263,8 @@ void FBrushDetails::CustomizeDetails( IDetailLayoutBuilder& DetailLayout )
 			.ButtonContent()
 			[
 				SNew(STextBlock)
-				.Text(NSLOCTEXT("BrushDetails", "OrderMenu", "Order").ToString())
-				.ToolTipText(NSLOCTEXT("BrushDetails", "OrderMenu_ToolTip", "Order options").ToString())
+				.Text(NSLOCTEXT("BrushDetails", "OrderMenu", "Order"))
+				.ToolTipText(NSLOCTEXT("BrushDetails", "OrderMenu_ToolTip", "Order options"))
 				.Font(IDetailLayoutBuilder::GetDetailFont())
 			]
 			.MenuContent()
@@ -282,12 +282,12 @@ void FBrushDetails::CustomizeDetails( IDetailLayoutBuilder& DetailLayout )
 		+SHorizontalBox::Slot()
 		[
 			SNew( SButton )
-			.ToolTipText( LOCTEXT("AlignBrushVerts_Tooltip", "Aligns each vertex of the brush to the grid.").ToString() )
+			.ToolTipText( LOCTEXT("AlignBrushVerts_Tooltip", "Aligns each vertex of the brush to the grid.") )
 			.OnClicked( FOnClicked::CreateStatic( &Local::ExecuteExecCommand, FString( TEXT("ACTOR ALIGN VERTS") ) ) )
 			.HAlign( HAlign_Center )
 			[
 				SNew( STextBlock )
-				.Text( LOCTEXT("AlignBrushVerts", "Align Brush Vertices").ToString() )
+				.Text( LOCTEXT("AlignBrushVerts", "Align Brush Vertices") )
 				.Font( IDetailLayoutBuilder::GetDetailFont() )
 			]
 		]
@@ -298,12 +298,12 @@ void FBrushDetails::CustomizeDetails( IDetailLayoutBuilder& DetailLayout )
 		BrushHorizontalBox->AddSlot()
 		[
 			SNew( SButton )
-			.ToolTipText( LOCTEXT("CreateStaticMeshActor_Tooltip", "Creates a static mesh from selected brushes and replaces the affected brushes in the scene with the new static mesh").ToString() )
+			.ToolTipText( LOCTEXT("CreateStaticMeshActor_Tooltip", "Creates a static mesh from selected brushes and replaces the affected brushes in the scene with the new static mesh") )
 			.OnClicked( this, &FBrushDetails::OnCreateStaticMesh )
 			.HAlign( HAlign_Center )
 			[
 				SNew( STextBlock )
-				.Text( LOCTEXT("CreateStaticMeshActor", "Create Static Mesh").ToString() )
+				.Text( LOCTEXT("CreateStaticMeshActor", "Create Static Mesh") )
 				.Font( IDetailLayoutBuilder::GetDetailFont() )
 			]
 		];
diff --git a/Engine/Source/Editor/DetailCustomizations/Private/CollisionProfileDetails.cpp b/Engine/Source/Editor/DetailCustomizations/Private/CollisionProfileDetails.cpp
--- a/Engine/Source/Editor/DetailCustomizations/Private/CollisionProfileDetails.cpp
+++ b/Engine/Source/Editor/DetailCustomizations/Private/CollisionProfileDetails.cpp
@@ -288,7 +288,7 @@ FText SChannelEditDialog::GetName() const
 TSharedRef<SWidget> SChannelEditDialog::HandleResponseComboBoxGenerateWidget(TSharedPtr<FString> StringItem)
 {
 	return SNew(STextBlock)
-		.Text(*StringItem)
+		.Text(FText::FromString(*StringItem))
 		.Font(IDetailLayoutBuilder::GetDetailFont());
 }
 
@@ -727,7 +727,7 @@ FText SProfileEditDialog::GetDescription() const
 TSharedRef<SWidget> SProfileEditDialog::HandleCollisionEnabledComboBoxGenerateWidget(TSharedPtr<FString> StringItem)
 {
 	return SNew(STextBlock)
-		.Text(*StringItem)
+		.Text(FText::FromString(*StringItem))
 		.Font(IDetailLayoutBuilder::GetDetailFont());
 }
 
@@ -763,7 +763,7 @@ FText SProfileEditDialog::HandleCollisionEnabledComboBoxContentText() const
 TSharedRef<SWidget> SProfileEditDialog::HandleObjectTypeComboBoxGenerateWidget(TSharedPtr<FString> StringItem)
 {
 	return SNew(STextBlock)
-		.Text(*StringItem)
+		.Text(FText::FromString(*StringItem))
 		.Font(IDetailLayoutBuilder::GetDetailFont());
 }
 
@@ -1033,7 +1033,7 @@ void SProfileEditDialog::AddCollisionChannel(TArray<FCollisionChannelInfo>	Valid
 				+SHorizontalBox::Slot()
 				[
 					SNew( STextBlock )
-					.Text(DisplayName)
+					.Text(FText::FromString(DisplayName))
 					.Font(IDetailLayoutBuilder::GetDetailFont())
 					.ToolTipText(LOCTEXT("SProfileEditDialog_CR_ToolTip", "When trace by channel, this information will be used for filtering."))
 				]
@@ -1138,19 +1138,19 @@ void SChannelListItem::Construct(const FArguments& InArgs, const TSharedRef<STab
 	SMultiColumnTableRow< TSharedPtr<FChannelListItem> >::Construct(FSuperRowType::FArguments(), InOwnerTableView);
 }
 
-FString SChannelListItem::GetDefaultResponse() const
+FText SChannelListItem::GetDefaultResponse() const
 {
 	switch (ChannelSetup->DefaultResponse)
 	{
 	case ECR_Ignore:
-		return TEXT("Ignore");
+		return LOCTEXT("ECR_Ignore", "Ignore");
 	case ECR_Overlap:
-		return TEXT("Overlap");
+		return LOCTEXT("ECR_Overlap", "Overlap");
 	case ECR_Block:
-		return TEXT("Block");
+		return LOCTEXT("ECR_Block", "Block");
 	}
 
-	return TEXT("ERROR");
+	return LOCTEXT("ECR_Error", "ERROR");
 }
 
 TSharedRef<SWidget> SChannelListItem::GenerateWidgetForColumn(const FName& ColumnName)
@@ -1215,24 +1215,24 @@ int32 UCollisionProfile::ReturnContainerIndexFromChannelName(FName& DisplayName)
 
 	return NameIndex;
 }
-FString SProfileListItem::GetObjectType() const
+FText SProfileListItem::GetObjectType() const
 {
-	return ProfileTemplate->ObjectTypeName.ToString();
+	return FText::FromName(ProfileTemplate->ObjectTypeName);
 }
 
-FString SProfileListItem::GetCollsionEnabled() const
+FText SProfileListItem::GetCollsionEnabled() const
 {
 	switch(ProfileTemplate->CollisionEnabled)
 	{
 	case ECollisionEnabled::NoCollision:
-		return TEXT("No Collision");
+		return LOCTEXT("ECollisionEnabled_NoCollision", "No Collision");
 	case ECollisionEnabled::QueryOnly:
-		return TEXT("No Physics Collision");
+		return LOCTEXT("ECollisionEnabled_QueryOnly", "No Physics Collision");
 	case ECollisionEnabled::QueryAndPhysics:
-		return TEXT("Collision Enabled");
+		return LOCTEXT("ECollisionEnabled_QueryAndPhysics", "Collision Enabled");
 	}
 
-	return TEXT("ERROR");
+	return LOCTEXT("ECollisionEnabled_Error", "ERROR");
 }
 
 TSharedRef<SWidget> SProfileListItem::GenerateWidgetForColumn(const FName& ColumnName)
@@ -1302,7 +1302,7 @@ TSharedRef<SWidget> SProfileListItem::GenerateWidgetForColumn(const FName& Colum
 			.VAlign(VAlign_Center)
 			[
 				SNew(STextBlock)
-				.Text(ProfileTemplate->HelpMessage)
+				.Text(FText::FromString(ProfileTemplate->HelpMessage))
 				.Font(IDetailLayoutBuilder::GetDetailFont())
 			];
 	}
diff --git a/Engine/Source/Editor/DetailCustomizations/Private/CollisionProfileDetails.h b/Engine/Source/Editor/DetailCustomizations/Private/CollisionProfileDetails.h
--- a/Engine/Source/Editor/DetailCustomizations/Private/CollisionProfileDetails.h
+++ b/Engine/Source/Editor/DetailCustomizations/Private/CollisionProfileDetails.h
@@ -44,7 +44,7 @@ public:
 
 private:
 
-	FString GetDefaultResponse() const;
+	FText GetDefaultResponse() const;
 
 	TSharedPtr<FCustomChannelSetup> ChannelSetup;
 };
@@ -91,8 +91,8 @@ public:
 
 private:
 
-	FString GetObjectType() const;
-	FString GetCollsionEnabled() const;
+	FText GetObjectType() const;
+	FText GetCollsionEnabled() const;
 	TSharedPtr<FCollisionResponseTemplate> ProfileTemplate;
 };
 
diff --git a/Engine/Source/Editor/DetailCustomizations/Private/CollisionProfileNameCustomization.cpp b/Engine/Source/Editor/DetailCustomizations/Private/CollisionProfileNameCustomization.cpp
--- a/Engine/Source/Editor/DetailCustomizations/Private/CollisionProfileNameCustomization.cpp
+++ b/Engine/Source/Editor/DetailCustomizations/Private/CollisionProfileNameCustomization.cpp
@@ -70,7 +70,7 @@ TSharedRef<SWidget> FCollisionProfileNameCustomization::OnGenerateWidget(TShared
 {
 	return
 		SNew(STextBlock)
-		.Text(InItem->ToString())
+		.Text(FText::FromName(*InItem))
 		.Font(IDetailLayoutBuilder::GetDetailFont());
 }
 
@@ -148,19 +148,19 @@ void FCollisionProfileNameCustomization::GetPropertyAsName(FName& OutName) const
 }
 
 
-FString FCollisionProfileNameCustomization::GetProfileComboBoxContent() const
+FText FCollisionProfileNameCustomization::GetProfileComboBoxContent() const
 {
 	TSharedPtr<FName> SelectedName = GetSelectedName();
 	if (SelectedName.IsValid())
 	{
-		return (*SelectedName->ToString());
+		return FText::FromName(*SelectedName);
 	}
 
-	return TEXT("Invalid");
+	return LOCTEXT("Invalid", "Invalid");
 }
 
 
-FString FCollisionProfileNameCustomization::GetProfileComboBoxToolTip() const
+FText FCollisionProfileNameCustomization::GetProfileComboBoxToolTip() const
 {
 	UCollisionProfile* CollisionProfile = UCollisionProfile::Get();
 	check(CollisionProfile);
@@ -171,10 +171,10 @@ FString FCollisionProfileNameCustomization::GetProfileComboBoxToolTip() const
 	FCollisionResponseTemplate ProfileTemplate;
 	if (CollisionProfile->GetProfileTemplate(ProfileName, ProfileTemplate))
 	{
-		return ProfileTemplate.HelpMessage;
+		return FText::FromString(ProfileTemplate.HelpMessage);
 	}
 
-	return TEXT("Invalid");
+	return LOCTEXT("Invalid", "Invalid");
 }
 
 
diff --git a/Engine/Source/Editor/DetailCustomizations/Private/CollisionProfileNameCustomization.h b/Engine/Source/Editor/DetailCustomizations/Private/CollisionProfileNameCustomization.h
--- a/Engine/Source/Editor/DetailCustomizations/Private/CollisionProfileNameCustomization.h
+++ b/Engine/Source/Editor/DetailCustomizations/Private/CollisionProfileNameCustomization.h
@@ -33,8 +33,8 @@ protected:
 	void SetPropertyWithName(const FName& Name);
 	void GetPropertyAsName(FName& OutName) const;
 
-	FString GetProfileComboBoxContent() const;
-	FString GetProfileComboBoxToolTip() const;
+	FText GetProfileComboBoxContent() const;
+	FText GetProfileComboBoxToolTip() const;
 
 protected:
 
diff --git a/Engine/Source/Editor/DetailCustomizations/Private/CurveTableCustomization.h b/Engine/Source/Editor/DetailCustomizations/Private/CurveTableCustomization.h
--- a/Engine/Source/Editor/DetailCustomizations/Private/CurveTableCustomization.h
+++ b/Engine/Source/Editor/DetailCustomizations/Private/CurveTableCustomization.h
@@ -125,28 +125,28 @@ private:
 		return
 			SNew( STableRow<TSharedPtr<FString>>, OwnerTable)
 			[
-				SNew( STextBlock ).Text( *InItem )
+				SNew( STextBlock ).Text( FText::FromString(*InItem) )
 			];
 	}
 
 	/** Display the current selection */
-	FString GetRowNameComboBoxContentText( ) const
+	FText GetRowNameComboBoxContentText( ) const
 	{
-		FString RowName = TEXT( "Multiple Values" );
-		const FPropertyAccess::Result RowResult = RowNamePropertyHandle->GetValue( RowName );
+		FString RowNameValue;
+		const FPropertyAccess::Result RowResult = RowNamePropertyHandle->GetValue( RowNameValue );
 		if ( RowResult != FPropertyAccess::MultipleValues )
 		{
 			TSharedPtr<FString> SelectedRowName = CurrentSelectedItem;
 			if ( SelectedRowName.IsValid() )
 			{
-				RowName = *SelectedRowName;
+				return FText::FromString(*SelectedRowName);
 			}
 			else
 			{
-				RowName = TEXT( "None" );
+				return NSLOCTEXT( "CurveTableCustomization", "None", "None" );
 			}
 		}
-		return RowName;
+		return NSLOCTEXT( "CurveTableCustomization", "MultipleValues", "Multiple Values" );
 	}
 
 	/** Update the root data on a change of selection */
diff --git a/Engine/Source/Editor/DetailCustomizations/Private/DataTableCategoryCustomization.h b/Engine/Source/Editor/DetailCustomizations/Private/DataTableCategoryCustomization.h
--- a/Engine/Source/Editor/DetailCustomizations/Private/DataTableCategoryCustomization.h
+++ b/Engine/Source/Editor/DetailCustomizations/Private/DataTableCategoryCustomization.h
@@ -206,54 +206,54 @@ private:
 	TSharedRef<SWidget> HandleColumnNameComboBoxGenarateWidget(TSharedPtr<FString> Item)
 	{
 		return SNew(STextBlock)
-			.Text(*Item);
+			.Text(FText::FromString(*Item));
 	}
 
 	/** Return the representation of the the row names to display */
 	TSharedRef<SWidget> HandleRowContentsComboBoxGenarateWidget(TSharedPtr<FString> Item)
 	{
 		return SNew(STextBlock)
-			.Text(*Item);
+			.Text(FText::FromString(*Item));
 	}
 	
 	/** Display the current column selection */
-	FString GetColumnNameComboBoxContentText() const
+	FText GetColumnNameComboBoxContentText() const
 	{
-		FString ColumnName = LOCTEXT("MultipleValues", "Multiple Values").ToString();
-		const FPropertyAccess::Result ColumnResult = ColumnNamePropertyHandle->GetValue(ColumnName);
+		FString ColumnNameValue;
+		const FPropertyAccess::Result ColumnResult = ColumnNamePropertyHandle->GetValue(ColumnNameValue);
 		if (ColumnResult != FPropertyAccess::MultipleValues)
 		{
 			TSharedPtr<FString> SelectedColumnName = ColumnNameComboBox->GetSelectedItem();
 			if (SelectedColumnName.IsValid())
 			{
-				ColumnName = *SelectedColumnName;
+				return FText::FromString(*SelectedColumnName);
 			}
 			else
 			{
-				ColumnName = LOCTEXT("DataTable_None", "None").ToString();
+				return LOCTEXT("DataTable_None", "None");
 			}
 		}
-		return ColumnName;
+		return LOCTEXT("MultipleValues", "Multiple Values");
 	}
 
 	/** Display the current row selection */
-	FString GetRowContentsComboBoxContentText() const
+	FText GetRowContentsComboBoxContentText() const
 	{
-		FString RowContains = LOCTEXT("MultipleValues", "Multiple Values").ToString();
-		const FPropertyAccess::Result RowResult = RowContentsPropertyHandle->GetValue(RowContains);
+		FString RowContainsValue;
+		const FPropertyAccess::Result RowResult = RowContentsPropertyHandle->GetValue(RowContainsValue);
 		if (RowResult != FPropertyAccess::MultipleValues)
 		{
 			TSharedPtr<FString> SelectedRowContents = RowContentsComboBox->GetSelectedItem();
 			if (SelectedRowContents.IsValid())
 			{
-				RowContains = *SelectedRowContents;
+				return FText::FromString(*SelectedRowContents);
 			}
 			else
 			{
-				RowContains = LOCTEXT("DataTable_None", "None").ToString();
+				return LOCTEXT("DataTable_None", "None");
 			}
 		}
-		return RowContains;
+		return LOCTEXT("MultipleValues", "Multiple Values");
 	}
 
 	/** Delegate to refresh the drop down when the datatable changes */
diff --git a/Engine/Source/Editor/DetailCustomizations/Private/DataTableCustomization.h b/Engine/Source/Editor/DetailCustomizations/Private/DataTableCustomization.h
--- a/Engine/Source/Editor/DetailCustomizations/Private/DataTableCustomization.h
+++ b/Engine/Source/Editor/DetailCustomizations/Private/DataTableCustomization.h
@@ -129,28 +129,28 @@ private:
 		return
 			SNew( STableRow<TSharedPtr<FString>>, OwnerTable)
 			[
-				SNew( STextBlock ).Text( *InItem )
+				SNew( STextBlock ).Text( FText::FromString(*InItem) )
 			];
 	}
 
 	/** Display the current selection */
-	FString GetRowNameComboBoxContentText( ) const
+	FText GetRowNameComboBoxContentText( ) const
 	{
-		FString RowName = LOCTEXT( "MultipleValues", "Multiple Values" ).ToString();
-		const FPropertyAccess::Result RowResult = RowNamePropertyHandle->GetValue( RowName );
+		FString RowNameValue;
+		const FPropertyAccess::Result RowResult = RowNamePropertyHandle->GetValue( RowNameValue );
 		if ( RowResult != FPropertyAccess::MultipleValues )
 		{
 			TSharedPtr<FString> SelectedRowName = CurrentSelectedItem;
 			if ( SelectedRowName.IsValid() )
 			{
-				RowName = *SelectedRowName;
+				return FText::FromString(*SelectedRowName);
 			}
 			else
 			{
-				RowName = LOCTEXT("DataTable_None", "None").ToString();
+				return LOCTEXT("DataTable_None", "None");
 			}
 		}
-		return RowName;
+		return LOCTEXT( "MultipleValues", "Multiple Values" );
 	}
 
 	/** Update the root data on a change of selection */
diff --git a/Engine/Source/Editor/DetailCustomizations/Private/DeviceProfileDetails.cpp b/Engine/Source/Editor/DetailCustomizations/Private/DeviceProfileDetails.cpp
--- a/Engine/Source/Editor/DetailCustomizations/Private/DeviceProfileDetails.cpp
+++ b/Engine/Source/Editor/DetailCustomizations/Private/DeviceProfileDetails.cpp
@@ -548,7 +548,7 @@ TSharedRef<SWidget> FDeviceProfileParentPropertyDetails::HandleDeviceProfilePare
 	.Padding(DeviceProfilePropertyConstants::CVarSelectionMenuPadding)
 	[
 		SNew(STextBlock)
-		.Text(*InItem)
+		.Text(FText::FromString(*InItem))
 	];
 }
 
diff --git a/Engine/Source/Editor/DetailCustomizations/Private/DialogueStructsCustomizations.cpp b/Engine/Source/Editor/DetailCustomizations/Private/DialogueStructsCustomizations.cpp
--- a/Engine/Source/Editor/DetailCustomizations/Private/DialogueStructsCustomizations.cpp
+++ b/Engine/Source/Editor/DetailCustomizations/Private/DialogueStructsCustomizations.cpp
@@ -265,7 +265,7 @@ TSharedRef<SWidget> SSpeakerDropDown::MakeComboButtonItemWidget(TSharedPtr<UDial
 		[
 			// Voice Description
 			SNew( STextBlock )
-			.Text( ( Speaker.IsValid() && *Speaker ) ? (*Speaker)->GetDesc() : LOCTEXT("None", "None").ToString() )
+			.Text( ( Speaker.IsValid() && *Speaker ) ? FText::FromString((*Speaker)->GetDesc()) : LOCTEXT("None", "None") )
 		];
 }
 
@@ -555,7 +555,7 @@ TSharedRef<SWidget> STargetSetDropDown::MakeComboButtonItemWidget(TSharedPtr<FTa
 			// Voice Description
 			SNew( STextBlock )
 			.Font( Font )
-			.Text( Target ? Target->GetDesc() : LOCTEXT("None", "None").ToString() )
+			.Text( Target ? FText::FromString(Target->GetDesc()) : LOCTEXT("None", "None") )
 		];
 	}
 	else
diff --git a/Engine/Source/Editor/DetailCustomizations/Private/DirectoryPathStructCustomization.cpp b/Engine/Source/Editor/DetailCustomizations/Private/DirectoryPathStructCustomization.cpp
--- a/Engine/Source/Editor/DetailCustomizations/Private/DirectoryPathStructCustomization.cpp
+++ b/Engine/Source/Editor/DetailCustomizations/Private/DirectoryPathStructCustomization.cpp
@@ -41,7 +41,7 @@ void FDirectoryPathStructCustomization::CustomizeHeader( TSharedRef<IPropertyHan
 			[
 				SAssignNew(BrowseButton, SButton)
 				.ButtonStyle( FEditorStyle::Get(), "HoverHintOnly" )
-				.ToolTipText( LOCTEXT( "FolderButtonToolTipText", "Choose a directory from this computer").ToString() )
+				.ToolTipText( LOCTEXT( "FolderButtonToolTipText", "Choose a directory from this computer") )
 				.OnClicked( FOnClicked::CreateSP(this, &FDirectoryPathStructCustomization::OnPickDirectory, PathProperty.ToSharedRef(), bRelativeToGameContentDir, bUseRelativePath) )
 				.ContentPadding( 2.0f )
 				.ForegroundColor( FSlateColor::UseForeground() )
diff --git a/Engine/Source/Editor/DetailCustomizations/Private/EnvQueryParamInstanceCustomization.cpp b/Engine/Source/Editor/DetailCustomizations/Private/EnvQueryParamInstanceCustomization.cpp
--- a/Engine/Source/Editor/DetailCustomizations/Private/EnvQueryParamInstanceCustomization.cpp
+++ b/Engine/Source/Editor/DetailCustomizations/Private/EnvQueryParamInstanceCustomization.cpp
@@ -145,20 +145,20 @@ EVisibility FEnvQueryParamInstanceCustomization::GetParamBoolValueVisibility() c
 	return (ParamType == EEnvQueryParam::Bool) ? EVisibility::Visible : EVisibility::Collapsed;
 }
 
-FString FEnvQueryParamInstanceCustomization::GetHeaderDesc() const
+FText FEnvQueryParamInstanceCustomization::GetHeaderDesc() const
 {
 	FString ParamName;
 	if (NameProp->GetValue(ParamName) == FPropertyAccess::Success)
 	{
 		switch (ParamType)
 		{
-		case EEnvQueryParam::Float:	return FString::Printf(TEXT("%s = %s"), *ParamName, *FString::SanitizeFloat(CachedFloat));
-		case EEnvQueryParam::Int:	return FString::Printf(TEXT("%s = %d"), *ParamName, CachedInt);
-		case EEnvQueryParam::Bool:	return FString::Printf(TEXT("%s = %s"), *ParamName, CachedBool ? TEXT("true") : TEXT("false"));
+		case EEnvQueryParam::Float:	return FText::FromString(FString::Printf(TEXT("%s = %s"), *ParamName, *FString::SanitizeFloat(CachedFloat)));
+		case EEnvQueryParam::Int:	return FText::FromString(FString::Printf(TEXT("%s = %d"), *ParamName, CachedInt));
+		case EEnvQueryParam::Bool:	return FText::FromString(FString::Printf(TEXT("%s = %s"), *ParamName, CachedBool ? TEXT("true") : TEXT("false")));
 		}
 	}
 
-	return FString();
+	return FText::GetEmpty();
 }
 
 void FEnvQueryParamInstanceCustomization::InitCachedTypes()
diff --git a/Engine/Source/Editor/DetailCustomizations/Private/EnvQueryParamInstanceCustomization.h b/Engine/Source/Editor/DetailCustomizations/Private/EnvQueryParamInstanceCustomization.h
--- a/Engine/Source/Editor/DetailCustomizations/Private/EnvQueryParamInstanceCustomization.h
+++ b/Engine/Source/Editor/DetailCustomizations/Private/EnvQueryParamInstanceCustomization.h
@@ -31,7 +31,7 @@ protected:
 	void OnParamBoolValueChanged(ECheckBoxState BoolValue) const;
 	EVisibility GetParamBoolValueVisibility() const;
 
-	FString GetHeaderDesc() const;
+	FText GetHeaderDesc() const;
 	void OnTypeChanged();
 	void InitCachedTypes();
 
diff --git a/Engine/Source/Editor/DetailCustomizations/Private/NavLinkStructCustomization.cpp b/Engine/Source/Editor/DetailCustomizations/Private/NavLinkStructCustomization.cpp
--- a/Engine/Source/Editor/DetailCustomizations/Private/NavLinkStructCustomization.cpp
+++ b/Engine/Source/Editor/DetailCustomizations/Private/NavLinkStructCustomization.cpp
@@ -30,7 +30,7 @@ void FNavLinkStructCustomization::CustomizeHeader( TSharedRef<IPropertyHandle> S
 	.MaxDesiredWidth(400.0f)
 	[
 		SNew(STextBlock)
-		.Text(Desc)
+		.Text(FText::FromString(Desc))
 		.Font(StructCustomizationUtils.GetRegularFont())
 	];
 }
diff --git a/Engine/Source/Editor/DetailCustomizations/Private/ParticleSysParamStructCustomization.cpp b/Engine/Source/Editor/DetailCustomizations/Private/ParticleSysParamStructCustomization.cpp
--- a/Engine/Source/Editor/DetailCustomizations/Private/ParticleSysParamStructCustomization.cpp
+++ b/Engine/Source/Editor/DetailCustomizations/Private/ParticleSysParamStructCustomization.cpp
@@ -168,7 +168,7 @@ TSharedRef<SWidget> FParticleSysParamStructCustomization::OnGenerateComboWidget(
 		SNew(SBox)
 		[
 			SNew(STextBlock)
-			.Text(*InComboString)
+			.Text(FText::FromString(*InComboString))
 			.ToolTipText(ToolTip)
 			.Font(IPropertyTypeCustomizationUtils::GetRegularFont())
 		];
diff --git a/Engine/Source/Editor/DetailCustomizations/Private/ParticleSystemComponentDetails.cpp b/Engine/Source/Editor/DetailCustomizations/Private/ParticleSystemComponentDetails.cpp
--- a/Engine/Source/Editor/DetailCustomizations/Private/ParticleSystemComponentDetails.cpp
+++ b/Engine/Source/Editor/DetailCustomizations/Private/ParticleSystemComponentDetails.cpp
@@ -28,22 +28,22 @@ void FParticleSystemComponentDetails::CustomizeDetails(IDetailLayoutBuilder& InD
 			[
 				SNew(SButton)
 				.OnClicked(this, &FParticleSystemComponentDetails::OnAutoPopulateClicked)
-				.ToolTipText(NSLOCTEXT("ParticleSystemComponentDetails", "AutoPopulateButtonTooltip", "Copies properties from the source particle system into the instanced parameters of this system").ToString())
+				.ToolTipText(NSLOCTEXT("ParticleSystemComponentDetails", "AutoPopulateButtonTooltip", "Copies properties from the source particle system into the instanced parameters of this system"))
 				.HAlign(HAlign_Center)
 				[
 					SNew(STextBlock)
-					.Text(NSLOCTEXT("ParticleSystemComponentDetails", "AutoPopulateButton", "Expose Parameter").ToString())
+					.Text(NSLOCTEXT("ParticleSystemComponentDetails", "AutoPopulateButton", "Expose Parameter"))
 				]
 			]
 			+ SUniformGridPanel::Slot(1, 0)
 				[
 					SNew(SButton)
 					.OnClicked(this, &FParticleSystemComponentDetails::OnResetEmitter)
-					.ToolTipText(NSLOCTEXT("ParticleSystemComponentDetails", "ResetEmitterButtonTooltip", "Resets the selected particle system.").ToString())
+					.ToolTipText(NSLOCTEXT("ParticleSystemComponentDetails", "ResetEmitterButtonTooltip", "Resets the selected particle system."))
 					.HAlign(HAlign_Center)
 					[
 						SNew(STextBlock)
-						.Text(NSLOCTEXT("ParticleSystemComponentDetails", "ResetEmitterButton", "Reset Emitter").ToString())
+						.Text(NSLOCTEXT("ParticleSystemComponentDetails", "ResetEmitterButton", "Reset Emitter"))
 					]
 				]
 		];
diff --git a/Engine/Source/Editor/DetailCustomizations/Private/PhysicsConstraintComponentDetails.cpp b/Engine/Source/Editor/DetailCustomizations/Private/PhysicsConstraintComponentDetails.cpp
--- a/Engine/Source/Editor/DetailCustomizations/Private/PhysicsConstraintComponentDetails.cpp
+++ b/Engine/Source/Editor/DetailCustomizations/Private/PhysicsConstraintComponentDetails.cpp
@@ -146,7 +146,7 @@ void FPhysicsConstraintComponentDetails::CustomizeDetails( IDetailLayoutBuilder&
 						.OnCheckStateChanged( this, &FPhysicsConstraintComponentDetails::OnLimitRadioChanged, CurProperty, LinearLimitEnum[0] )
 						.ToolTipText(LinearLimitOptionTooltips[0])
 						[
-							SNew(STextBlock).Text(*LinearLimitOptionNames[0].Get())
+							SNew(STextBlock).Text(FText::FromString(*LinearLimitOptionNames[0].Get()))
 						]						
 					]
 					+SVerticalBox::Slot()
@@ -159,7 +159,7 @@ void FPhysicsConstraintComponentDetails::CustomizeDetails( IDetailLayoutBuilder&
 							.OnCheckStateChanged( this, &FPhysicsConstraintComponentDetails::OnLimitRadioChanged, CurProperty, LinearLimitEnum[1] )
 							.ToolTipText(LinearLimitOptionTooltips[1])
 							[
-								SNew(STextBlock).Text(*LinearLimitOptionNames[1].Get())
+								SNew(STextBlock).Text(FText::FromString(*LinearLimitOptionNames[1].Get()))
 							]
 						]
 					+SVerticalBox::Slot()
@@ -172,7 +172,7 @@ void FPhysicsConstraintComponentDetails::CustomizeDetails( IDetailLayoutBuilder&
 							.OnCheckStateChanged( this, &FPhysicsConstraintComponentDetails::OnLimitRadioChanged, CurProperty, LinearLimitEnum[2] )
 							.ToolTipText(LinearLimitOptionTooltips[2])
 							[
-								SNew(STextBlock).Text(*LinearLimitOptionNames[2].Get())
+								SNew(STextBlock).Text(FText::FromString(*LinearLimitOptionNames[2].Get()))
 							]
 						]
 				];
@@ -239,7 +239,7 @@ void FPhysicsConstraintComponentDetails::CustomizeDetails( IDetailLayoutBuilder&
 						.OnCheckStateChanged( this, &FPhysicsConstraintComponentDetails::OnLimitRadioChanged, CurProperty, AngularLimitEnum[0] )
 						.ToolTipText(AngularLimitOptionTooltips[0])
 						[
-							SNew(STextBlock).Text(*AngularLimitOptionNames[0].Get())
+							SNew(STextBlock).Text(FText::FromString(*AngularLimitOptionNames[0].Get()))
 						]						
 					]
 					+SVerticalBox::Slot()
@@ -252,7 +252,7 @@ void FPhysicsConstraintComponentDetails::CustomizeDetails( IDetailLayoutBuilder&
 							.OnCheckStateChanged( this, &FPhysicsConstraintComponentDetails::OnLimitRadioChanged, CurProperty, AngularLimitEnum[1] )
 							.ToolTipText(AngularLimitOptionTooltips[1])
 							[
-								SNew(STextBlock).Text(*AngularLimitOptionNames[1].Get())
+								SNew(STextBlock).Text(FText::FromString(*AngularLimitOptionNames[1].Get()))
 							]
 						]
 					+SVerticalBox::Slot()
@@ -265,7 +265,7 @@ void FPhysicsConstraintComponentDetails::CustomizeDetails( IDetailLayoutBuilder&
 							.OnCheckStateChanged( this, &FPhysicsConstraintComponentDetails::OnLimitRadioChanged, CurProperty, AngularLimitEnum[2] )
 							.ToolTipText(AngularLimitOptionTooltips[2])
 							[
-								SNew(STextBlock).Text(*AngularLimitOptionNames[2].Get())
+								SNew(STextBlock).Text(FText::FromString(*AngularLimitOptionNames[2].Get()))
 							]
 						]
 				];
diff --git a/Engine/Source/Editor/DetailCustomizations/Private/PhysicsSettingsDetails.cpp b/Engine/Source/Editor/DetailCustomizations/Private/PhysicsSettingsDetails.cpp
--- a/Engine/Source/Editor/DetailCustomizations/Private/PhysicsSettingsDetails.cpp
+++ b/Engine/Source/Editor/DetailCustomizations/Private/PhysicsSettingsDetails.cpp
@@ -40,7 +40,7 @@ TSharedRef<SWidget> SPhysicalSurfaceListItem::GenerateWidgetForColumn(const FNam
 			.VAlign(VAlign_Center)
 			[
 				SNew(STextBlock)
-				.Text(FString::Printf(TEXT("%s"), *GetTypeString()))
+				.Text(FText::FromString(GetTypeString()))
 				.Font(IDetailLayoutBuilder::GetDetailFont())
 			];
 	}
diff --git a/Engine/Source/Editor/DetailCustomizations/Private/RangeStructCustomization.cpp b/Engine/Source/Editor/DetailCustomizations/Private/RangeStructCustomization.cpp
--- a/Engine/Source/Editor/DetailCustomizations/Private/RangeStructCustomization.cpp
+++ b/Engine/Source/Editor/DetailCustomizations/Private/RangeStructCustomization.cpp
@@ -358,7 +358,7 @@ TSharedRef<SWidget> FRangeStructCustomization<NumericType>::OnGenerateComboWidge
 		.WidthOverride(150.0f)
 		[
 			SNew(STextBlock)
-			.Text(*InComboString)
+			.Text(FText::FromString(*InComboString))
 			.ToolTipText(ToolTip)
 			.Font(IDetailLayoutBuilder::GetDetailFont())
 			.IsEnabled(true)
diff --git a/Engine/Source/Editor/DetailCustomizations/Private/RigDetails.cpp b/Engine/Source/Editor/DetailCustomizations/Private/RigDetails.cpp
--- a/Engine/Source/Editor/DetailCustomizations/Private/RigDetails.cpp
+++ b/Engine/Source/Editor/DetailCustomizations/Private/RigDetails.cpp
@@ -112,10 +112,10 @@ void FRigDetails::GenerateNodeArrayElementWidget(TSharedRef<IPropertyHandle> Pro
 
 	// the interface will be node [display name] [parent node]
 	// delegate for display name
-	FString NodeName, ParentNodeName, DisplayString;
-	check (NodeNameProp->GetValueAsDisplayString(NodeName) != FPropertyAccess::Fail);
-	check (ParentNameProp->GetValueAsDisplayString(ParentNodeName) != FPropertyAccess::Fail);
-	check (DisplayNameProp->GetValueAsDisplayString(DisplayString) != FPropertyAccess::Fail);
+	FText NodeName, ParentNodeName, DisplayString;
+	check (NodeNameProp->GetValueAsDisplayText(NodeName) != FPropertyAccess::Fail);
+	check (ParentNameProp->GetValueAsDisplayText(ParentNodeName) != FPropertyAccess::Fail);
+	check (DisplayNameProp->GetValueAsDisplayText(DisplayString) != FPropertyAccess::Fail);
 
 	ChildrenBuilder.AddChildContent(FText::GetEmpty())
 	[
@@ -144,7 +144,7 @@ void FRigDetails::GenerateNodeArrayElementWidget(TSharedRef<IPropertyHandle> Pro
 			.Content()
 			[
 				SNew(STextBlock)
-				.Text(FString::Printf(TEXT(" [Parent : %s] "), *ParentNodeName))
+				.Text(FText::Format(LOCTEXT("ParentNameFmt", " [Parent : {0}] "), ParentNodeName))
 				.Font(DetailLayout->GetDetailFont())
 			]
 		]
@@ -154,7 +154,7 @@ void FRigDetails::GenerateNodeArrayElementWidget(TSharedRef<IPropertyHandle> Pro
 		.AutoWidth()
 		[
 			SNew(STextBlock)
-			.Text(FString(TEXT("Display Name")))
+			.Text(LOCTEXT("DisplayNameLabel", "Display Name"))
 			.Font(DetailLayout->GetDetailFontBold())
 		]
 
@@ -181,7 +181,7 @@ void FRigDetails::GenerateNodeArrayElementWidget(TSharedRef<IPropertyHandle> Pro
 		.AutoWidth()
 		[
 			SNew(STextBlock)
-			.Text(FString(TEXT("Advanced")))
+			.Text(LOCTEXT("AdvancedLabel", "Advanced"))
 			.Font(DetailLayout->GetDetailFontBold())
 		]
 
@@ -273,8 +273,9 @@ void FRigDetails::GenerateTransformBaseArrayElementWidget(TSharedRef<IPropertyHa
 
 	// the interface will be node [display name] [parent node]
 	// delegate for display name
-	FString NodeName, ParentNodeName_T, ParentNodeName_R;
-	check(NodeNameProp->GetValueAsDisplayString(NodeName) != FPropertyAccess::Fail);
+	FText NodeName;
+	FString ParentNodeName_T, ParentNodeName_R;
+	check(NodeNameProp->GetValueAsDisplayText(NodeName) != FPropertyAccess::Fail);
 	check(ParentNameProp_T->GetValueAsDisplayString(ParentNodeName_T) != FPropertyAccess::Fail);
 	check(ParentNameProp_R->GetValueAsDisplayString(ParentNodeName_R) != FPropertyAccess::Fail);
 
@@ -352,7 +353,7 @@ void FRigDetails::GenerateTransformBaseArrayElementWidget(TSharedRef<IPropertyHa
 					.Content()
 					[
 						SNew(STextBlock)
-						.Text(FString(TEXT("Translation")))
+						.Text(LOCTEXT("TranslationLabel", "Translation"))
 						.Font(DetailLayout->GetDetailFontBold())
 					]
 				]
@@ -395,7 +396,7 @@ void FRigDetails::GenerateTransformBaseArrayElementWidget(TSharedRef<IPropertyHa
 					.Content()
 					[
 						SNew(STextBlock)
-						.Text(FString(TEXT("Orientation")))
+						.Text(LOCTEXT("OrientationLabel", "Orientation"))
 						.Font(DetailLayout->GetDetailFontBold())
 					]
 				]
@@ -469,24 +470,24 @@ TSharedRef<SWidget> FRigDetails::MakeItemWidget(TSharedPtr<FString> StringItem)
 	check(StringItem.IsValid());
 
 	return SNew(STextBlock)
-		.Text(*StringItem.Get());
+		.Text(FText::FromString(*StringItem.Get()));
 }
 /** Helper method to get the text for a given item in the combo box */
-FString FRigDetails::GetSelectedTextLabel(TSharedRef<IPropertyHandle> ParentSpacePropertyHandle) const
+FText FRigDetails::GetSelectedTextLabel(TSharedRef<IPropertyHandle> ParentSpacePropertyHandle) const
 {
 	FString DisplayText;
 
 	if (ParentSpacePropertyHandle->GetValueAsDisplayString(DisplayText) != FPropertyAccess::Fail)
 	{
-		return DisplayText;
+		return FText::FromString(DisplayText);
 	}
 
-	return TEXT("Unknown");
+	return LOCTEXT("Unknown", "Unknown");
 }
 
 void FRigDetails::OnComboBoxOopening(TSharedRef<IPropertyHandle> ParentSpacePropertyHandle, int32 ArrayIndex, bool bTranslation)
 {
-	FString PropertyValue = GetSelectedTextLabel(ParentSpacePropertyHandle);
+	FString PropertyValue = GetSelectedTextLabel(ParentSpacePropertyHandle).ToString();
 
 	// now find exact data
 	TArray<TSharedPtr<FString>> & ParentOptions = ParentSpaceOptionList[ArrayIndex];
diff --git a/Engine/Source/Editor/DetailCustomizations/Private/RigDetails.h b/Engine/Source/Editor/DetailCustomizations/Private/RigDetails.h
--- a/Engine/Source/Editor/DetailCustomizations/Private/RigDetails.h
+++ b/Engine/Source/Editor/DetailCustomizations/Private/RigDetails.h
@@ -41,7 +41,7 @@ private:
 	/** Called to create a widget for each string */
 	TSharedRef<SWidget> MakeItemWidget(TSharedPtr<FString> StringItem);
 	// Helper method to get the text for a given item in the combo box 
-	FString GetSelectedTextLabel(TSharedRef<IPropertyHandle> ParentSpacePropertyHandle) const;
+	FText GetSelectedTextLabel(TSharedRef<IPropertyHandle> ParentSpacePropertyHandle) const;
 	void OnComboBoxOopening(TSharedRef<IPropertyHandle> ParentSpacePropertyHandle, int32 ArrayIndex, bool bTranslation);
 	// check box handler
 	// Callback for changing this row's Share check box state.
diff --git a/Engine/Source/Editor/DetailCustomizations/Private/SkeletalMeshComponentDetails.cpp b/Engine/Source/Editor/DetailCustomizations/Private/SkeletalMeshComponentDetails.cpp
--- a/Engine/Source/Editor/DetailCustomizations/Private/SkeletalMeshComponentDetails.cpp
+++ b/Engine/Source/Editor/DetailCustomizations/Private/SkeletalMeshComponentDetails.cpp
@@ -331,7 +331,7 @@ TSharedRef<SWidget> FSkeletalMeshComponentDetails::GetClassPickerMenuContent()
 		];
 }
 
-FString FSkeletalMeshComponentDetails::GetSelectedAnimBlueprintName() const
+FText FSkeletalMeshComponentDetails::GetSelectedAnimBlueprintName() const
 {
 	check(AnimationBlueprintHandle->IsValidHandle());
 
@@ -339,11 +339,11 @@ FString FSkeletalMeshComponentDetails::GetSelectedAnimBlueprintName() const
 	AnimationBlueprintHandle->GetValue(Object);
 	if(Object)
 	{
-		return Object->GetName();
+		return FText::FromString(Object->GetName());
 	}
 	else
 	{
-		return FString(TEXT("None"));
+		return LOCTEXT("None", "None");
 	}
 }
 
diff --git a/Engine/Source/Editor/DetailCustomizations/Private/SkeletalMeshComponentDetails.h b/Engine/Source/Editor/DetailCustomizations/Private/SkeletalMeshComponentDetails.h
--- a/Engine/Source/Editor/DetailCustomizations/Private/SkeletalMeshComponentDetails.h
+++ b/Engine/Source/Editor/DetailCustomizations/Private/SkeletalMeshComponentDetails.h
@@ -43,7 +43,7 @@ private:
 	TSharedRef<SWidget> GetClassPickerMenuContent();
 
 	/** Gets the currently selected blueprint name to display on the class picker combo button */
-	FString GetSelectedAnimBlueprintName() const;
+	FText GetSelectedAnimBlueprintName() const;
 
 	/** Callback from the class picker when the user selects a class */
 	void OnClassPicked(UClass* PickedClass);
diff --git a/Engine/Source/Editor/DetailCustomizations/Private/SkeletonNotifyDetails.cpp b/Engine/Source/Editor/DetailCustomizations/Private/SkeletonNotifyDetails.cpp
--- a/Engine/Source/Editor/DetailCustomizations/Private/SkeletonNotifyDetails.cpp
+++ b/Engine/Source/Editor/DetailCustomizations/Private/SkeletonNotifyDetails.cpp
@@ -55,7 +55,7 @@ TSharedRef< ITableRow > FSkeletonNotifyDetails::MakeAnimationRow( TSharedPtr<FSt
 {
 	return SNew(STableRow<TSharedPtr<FString>>, OwnerTable)
 	[
-		SNew(STextBlock).Text(*Item.Get())
+		SNew(STextBlock).Text(FText::FromString(*Item.Get()))
 	];
 }
 
diff --git a/Engine/Source/Editor/DetailCustomizations/Private/SlateBrushCustomization.cpp b/Engine/Source/Editor/DetailCustomizations/Private/SlateBrushCustomization.cpp
--- a/Engine/Source/Editor/DetailCustomizations/Private/SlateBrushCustomization.cpp
+++ b/Engine/Source/Editor/DetailCustomizations/Private/SlateBrushCustomization.cpp
@@ -399,26 +399,26 @@ private:
 	/**
 	 * Return the text for the specified horizontal alignment
 	 */
-	FString MakeHorizontalAlignmentComboText( EHorizontalAlignment Alignment ) const
+	FText MakeHorizontalAlignmentComboText( EHorizontalAlignment Alignment ) const
 	{
-		FString AlignmentText;
+		FText AlignmentText;
 	
 		switch( Alignment )
 		{
 			case EHorizontalAlignment::HAlign_Fill:
-				AlignmentText = NSLOCTEXT( "UnrealEd", "AlignmentFill", "Fill" ).ToString();
+				AlignmentText = NSLOCTEXT( "UnrealEd", "AlignmentFill", "Fill" );
 				break;
 	
 			case EHorizontalAlignment::HAlign_Left:
-				AlignmentText = NSLOCTEXT( "UnrealEd", "AlignmentLeft", "Left" ).ToString();
+				AlignmentText = NSLOCTEXT( "UnrealEd", "AlignmentLeft", "Left" );
 				break;
 	
 			case EHorizontalAlignment::HAlign_Center:
-				AlignmentText = NSLOCTEXT( "UnrealEd", "AlignmentCenter", "Center" ).ToString();
+				AlignmentText = NSLOCTEXT( "UnrealEd", "AlignmentCenter", "Center" );
 				break;
 	
 			case EHorizontalAlignment::HAlign_Right:
-				AlignmentText = NSLOCTEXT( "UnrealEd", "AlignmentRight", "Right" ).ToString();
+				AlignmentText = NSLOCTEXT( "UnrealEd", "AlignmentRight", "Right" );
 				break;
 		}
 	
@@ -428,26 +428,26 @@ private:
 	/**
 	 * Return the text for the specified vertical alignment
 	 */
-	FString MakeVerticalAlignmentComboText( EVerticalAlignment Alignment ) const
+	FText MakeVerticalAlignmentComboText( EVerticalAlignment Alignment ) const
 	{
-		FString AlignmentText;
+		FText AlignmentText;
 	
 		switch( Alignment )
 		{
 			case EVerticalAlignment::VAlign_Fill:
-				AlignmentText = NSLOCTEXT( "UnrealEd", "AlignmentFill", "Fill" ).ToString();
+				AlignmentText = NSLOCTEXT( "UnrealEd", "AlignmentFill", "Fill" );
 				break;
 	
 			case EVerticalAlignment::VAlign_Top:
-				AlignmentText = NSLOCTEXT( "UnrealEd", "AlignmentTop", "Top" ).ToString();
+				AlignmentText = NSLOCTEXT( "UnrealEd", "AlignmentTop", "Top" );
 				break;
 	
 			case EVerticalAlignment::VAlign_Center:
-				AlignmentText = NSLOCTEXT( "UnrealEd", "AlignmentCenter", "Center" ).ToString();
+				AlignmentText = NSLOCTEXT( "UnrealEd", "AlignmentCenter", "Center" );
 				break;
 	
 			case EVerticalAlignment::VAlign_Bottom:
-				AlignmentText = NSLOCTEXT( "UnrealEd", "AlignmentBottom", "Bottom" ).ToString();
+				AlignmentText = NSLOCTEXT( "UnrealEd", "AlignmentBottom", "Bottom" );
 				break;
 		}
 	
@@ -457,26 +457,26 @@ private:
 	/**
 	 * Return the tooltip text for the specified horizontal alignment
 	 */
-	FString MakeHorizontalAlignmentComboToolTipText( EHorizontalAlignment Alignment ) const
+	FText MakeHorizontalAlignmentComboToolTipText( EHorizontalAlignment Alignment ) const
 	{
-		FString ToolTipText;
+		FText ToolTipText;
 	
 		switch( Alignment )
 		{
 			case EHorizontalAlignment::HAlign_Fill:
-				ToolTipText = NSLOCTEXT( "UnrealEd", "AlignmentFillToolTip", "The image will fill the preview" ).ToString();
+				ToolTipText = NSLOCTEXT( "UnrealEd", "AlignmentFillToolTip", "The image will fill the preview" );
 				break;
 	
 			case EHorizontalAlignment::HAlign_Left:
-				ToolTipText = NSLOCTEXT( "UnrealEd", "AlignmentLeftToolTip", "The image will be aligned to the left of the preview" ).ToString();
+				ToolTipText = NSLOCTEXT( "UnrealEd", "AlignmentLeftToolTip", "The image will be aligned to the left of the preview" );
 				break;
 	
 			case EHorizontalAlignment::HAlign_Center:
-				ToolTipText = NSLOCTEXT( "UnrealEd", "AlignmentCenterToolTip", "The image will be positioned in the centre of the preview" ).ToString();
+				ToolTipText = NSLOCTEXT( "UnrealEd", "AlignmentCenterToolTip", "The image will be positioned in the centre of the preview" );
 				break;
 	
 			case EHorizontalAlignment::HAlign_Right:
-				ToolTipText = NSLOCTEXT( "UnrealEd", "AlignmentRightToolTip", "The image will be aligned from the right of the preview" ).ToString();
+				ToolTipText = NSLOCTEXT( "UnrealEd", "AlignmentRightToolTip", "The image will be aligned from the right of the preview" );
 				break;
 		}
 	
@@ -486,26 +486,26 @@ private:
 	/**
 	 * Return the tooltip text for the specified vertical alignment
 	 */
-	FString MakeVerticalAlignmentComboToolTipText( EVerticalAlignment Alignment ) const
+	FText MakeVerticalAlignmentComboToolTipText( EVerticalAlignment Alignment ) const
 	{
-		FString ToolTipText;
+		FText ToolTipText;
 	
 		switch( Alignment )
 		{
 			case EVerticalAlignment::VAlign_Fill:
-				ToolTipText = NSLOCTEXT( "UnrealEd", "AlignmentFillToolTip", "The image will fill the preview" ).ToString();
+				ToolTipText = NSLOCTEXT( "UnrealEd", "AlignmentFillToolTip", "The image will fill the preview" );
 				break;
 	
 			case EVerticalAlignment::VAlign_Top:
-				ToolTipText = NSLOCTEXT( "UnrealEd", "AlignmentTopToolTip", "The image will be aligned to the top of the preview" ).ToString();
+				ToolTipText = NSLOCTEXT( "UnrealEd", "AlignmentTopToolTip", "The image will be aligned to the top of the preview" );
 				break;
 	
 			case EVerticalAlignment::VAlign_Center:
-				ToolTipText = NSLOCTEXT( "UnrealEd", "AlignmentCenterToolTip", "The image will be positioned in the centre of the preview" ).ToString();
+				ToolTipText = NSLOCTEXT( "UnrealEd", "AlignmentCenterToolTip", "The image will be positioned in the centre of the preview" );
 				break;
 	
 			case EVerticalAlignment::VAlign_Bottom:
-				ToolTipText = NSLOCTEXT( "UnrealEd", "AlignmentBottomToolTip", "The image will be aligned from the bottom of the preview" ).ToString();
+				ToolTipText = NSLOCTEXT( "UnrealEd", "AlignmentBottomToolTip", "The image will be aligned from the bottom of the preview" );
 				break;
 		}
 	
@@ -539,7 +539,7 @@ private:
 	/**
 	 * Get the horizontal alignment combo box content text
 	 */
-	FString GetHorizontalAlignmentComboBoxContent() const
+	FText GetHorizontalAlignmentComboBoxContent() const
 	{
 		return MakeHorizontalAlignmentComboText( HorizontalAlignment );
 	}
@@ -547,7 +547,7 @@ private:
 	/**
 	 * Get the vertical alignment combo box content text
 	 */
-	FString GetVerticalAlignmentComboBoxContent() const
+	FText GetVerticalAlignmentComboBoxContent() const
 	{
 		return MakeVerticalAlignmentComboText( VerticalAlignment );
 	}
@@ -555,7 +555,7 @@ private:
 	/**
 	 * Get the horizontal alignment combo box content tooltip text
 	 */
-	FString GetHorizontalAlignmentComboBoxContentToolTip() const
+	FText GetHorizontalAlignmentComboBoxContentToolTip() const
 	{
 		return MakeHorizontalAlignmentComboToolTipText( HorizontalAlignment );
 	}
@@ -563,7 +563,7 @@ private:
 	/**
 	 * Get the vertical alignment combo box content tooltip text
 	 */
-	FString GetVerticalAlignmentComboBoxContentToolTip() const
+	FText GetVerticalAlignmentComboBoxContentToolTip() const
 	{
 		return MakeVerticalAlignmentComboToolTipText( VerticalAlignment );
 	}
diff --git a/Engine/Source/Editor/DeviceProfileEditor/Private/DetailsPanel/SDeviceProfileDetailsPanel.cpp b/Engine/Source/Editor/DeviceProfileEditor/Private/DetailsPanel/SDeviceProfileDetailsPanel.cpp
--- a/Engine/Source/Editor/DeviceProfileEditor/Private/DetailsPanel/SDeviceProfileDetailsPanel.cpp
+++ b/Engine/Source/Editor/DeviceProfileEditor/Private/DetailsPanel/SDeviceProfileDetailsPanel.cpp
@@ -118,7 +118,7 @@ void SDeviceProfileDetailsPanel::RefreshUI()
 						.VAlign(VAlign_Center)
 						[
 							SNew( STextBlock )
-							.Text( FString::Printf( TEXT( "%s selected" ), *ViewingProfile->GetName() ) )
+							.Text( FText::Format( LOCTEXT( "SelectedDeviceProfileFmt", "{0} selected" ), FText::FromString(ViewingProfile->GetName()) ) )
 						]
 					]
 				]
diff --git a/Engine/Source/Editor/DeviceProfileEditor/Private/SDeviceProfileCreateProfilePanel.cpp b/Engine/Source/Editor/DeviceProfileEditor/Private/SDeviceProfileCreateProfilePanel.cpp
--- a/Engine/Source/Editor/DeviceProfileEditor/Private/SDeviceProfileCreateProfilePanel.cpp
+++ b/Engine/Source/Editor/DeviceProfileEditor/Private/SDeviceProfileCreateProfilePanel.cpp
@@ -206,9 +206,9 @@ END_SLATE_FUNCTION_BUILD_OPTIMIZATION
 
 
 
-FString SDeviceProfileCreateProfilePanel::SetProfileTypeComboBoxContent() const
+FText SDeviceProfileCreateProfilePanel::SetProfileTypeComboBoxContent() const
 {
-	return SelectedDeviceProfileType.IsValid() ? *SelectedDeviceProfileType : LOCTEXT("SelectType", "Choose a device profile type...").ToString();
+	return SelectedDeviceProfileType.IsValid() ? FText::FromString(*SelectedDeviceProfileType) : LOCTEXT("SelectType", "Choose a device profile type...");
 }
 
 
@@ -242,7 +242,7 @@ TSharedRef<SWidget> SDeviceProfileCreateProfilePanel::HandleProfileTypeComboBoxG
 	.Padding(DeviceProfileCreateProfileUIConstants::ListElementPadding)
 	[
 		SNew(STextBlock)
-		.Text(*InItem)
+		.Text(FText::FromString(*InItem))
 	];
 }
 
@@ -260,9 +260,9 @@ void SDeviceProfileCreateProfilePanel::HandleBaseProfileSelectionChanged( UDevic
 }
 
 
-FString SDeviceProfileCreateProfilePanel::SetBaseProfileComboBoxContent() const
+FText SDeviceProfileCreateProfilePanel::SetBaseProfileComboBoxContent() const
 {
-	return SelectedDeviceProfileParent.IsValid() ? SelectedDeviceProfileParent->GetName() : LOCTEXT("SelectParent", "Copy properties from...").ToString();
+	return SelectedDeviceProfileParent.IsValid() ? FText::FromString(SelectedDeviceProfileParent->GetName()) : LOCTEXT("SelectParent", "Copy properties from...");
 }
 
 
@@ -272,7 +272,7 @@ TSharedRef<SWidget> SDeviceProfileCreateProfilePanel::HandleBaseComboBoxGenerate
 	.Padding(DeviceProfileCreateProfileUIConstants::ListElementPadding)
 	[
 		SNew(STextBlock)
-		.Text(InItem->GetName())
+		.Text(FText::FromString(InItem->GetName()))
 	];
 }
 
diff --git a/Engine/Source/Editor/DeviceProfileEditor/Private/SDeviceProfileCreateProfilePanel.h b/Engine/Source/Editor/DeviceProfileEditor/Private/SDeviceProfileCreateProfilePanel.h
--- a/Engine/Source/Editor/DeviceProfileEditor/Private/SDeviceProfileCreateProfilePanel.h
+++ b/Engine/Source/Editor/DeviceProfileEditor/Private/SDeviceProfileCreateProfilePanel.h
@@ -58,7 +58,7 @@ public:
 	 *
 	 * @return The device profile parent name.
 	 */
-	FString SetBaseProfileComboBoxContent() const;
+	FText SetBaseProfileComboBoxContent() const;
 
 	/**
 	 * Handle base combo box generate widget.
@@ -83,7 +83,7 @@ public:
 	 *
 	 * @return The name of the class the created profile will be.
 	 */
-	FString SetProfileTypeComboBoxContent() const;
+	FText SetProfileTypeComboBoxContent() const;
 
 	/**
 	 * Handle profile type combo box generate widget.
diff --git a/Engine/Source/Editor/DeviceProfileEditor/Private/SDeviceProfileSelectionPanel.cpp b/Engine/Source/Editor/DeviceProfileEditor/Private/SDeviceProfileSelectionPanel.cpp
--- a/Engine/Source/Editor/DeviceProfileEditor/Private/SDeviceProfileSelectionPanel.cpp
+++ b/Engine/Source/Editor/DeviceProfileEditor/Private/SDeviceProfileSelectionPanel.cpp
@@ -72,7 +72,7 @@ public:
 	 *
 	 * @return The profile display name.
 	 */
-	FString GetProfileDisplayName() const;
+	FText GetProfileDisplayName() const;
 
 private:
 
@@ -131,7 +131,7 @@ TSharedRef< SWidget > SDeviceProfileSelectionRow::GenerateWidgetForColumn( const
 		// Draw a pin to show the state of the profile selection
 		ColumnWidget = SAssignNew( PinProfileButton, SButton )
 			.IsFocusable( false )
-			.ToolTipText(LOCTEXT("PinProfileColumnButtonToolTip", "Pin profile to device profile editor table").ToString())
+			.ToolTipText(LOCTEXT("PinProfileColumnButtonToolTip", "Pin profile to device profile editor table"))
 			.ButtonStyle( FEditorStyle::Get(), "NoBorder" )
 			.ContentPadding( 0 ) 
 			.HAlign( HAlign_Center )
@@ -152,7 +152,7 @@ TSharedRef< SWidget > SDeviceProfileSelectionRow::GenerateWidgetForColumn( const
 	{
 		ColumnWidget = SAssignNew(ViewProfileButton, SButton)
 			.IsFocusable(false)
-			.ToolTipText(LOCTEXT("ViewSingleProfileColumnButtonToolTip", "View this profile in it's own editor").ToString())
+			.ToolTipText(LOCTEXT("ViewSingleProfileColumnButtonToolTip", "View this profile in it's own editor"))
 			.ButtonStyle(FEditorStyle::Get(), "NoBorder")
 			.ContentPadding(0)
 			.HAlign(HAlign_Center)
@@ -199,9 +199,9 @@ const FSlateBrush* SDeviceProfileSelectionRow::GetPinnedImage() const
 }
 
 
-FString SDeviceProfileSelectionRow::GetProfileDisplayName() const
+FText SDeviceProfileSelectionRow::GetProfileDisplayName() const
 {
-	return SelectedDeviceProfile.IsValid() ? SelectedDeviceProfile->GetName() : TEXT("InvalidProfile");
+	return SelectedDeviceProfile.IsValid() ? FText::FromString(SelectedDeviceProfile->GetName()) : LOCTEXT("InvalidProfile", "Invalid Profile");
 }
 
 
diff --git a/Engine/Source/Editor/EditorWidgets/Private/SAssetDiscoveryIndicator.cpp b/Engine/Source/Editor/EditorWidgets/Private/SAssetDiscoveryIndicator.cpp
--- a/Engine/Source/Editor/EditorWidgets/Private/SAssetDiscoveryIndicator.cpp
+++ b/Engine/Source/Editor/EditorWidgets/Private/SAssetDiscoveryIndicator.cpp
@@ -67,7 +67,7 @@ void SAssetDiscoveryIndicator::Construct( const FArguments& InArgs )
 			[
 				SNew(STextBlock)
 				.Font( FEditorStyle::GetFontStyle("AssetDiscoveryIndicator.DiscovertingAssetsFont") )
-				.Text( NSLOCTEXT("AssetDiscoveryIndicator", "DiscoveringAssets", "Discovering Assets").ToString() )
+				.Text( NSLOCTEXT("AssetDiscoveryIndicator", "DiscoveringAssets", "Discovering Assets") )
 				.AutoWrapText(true)
 				.Justification(ETextJustify::Center)
 			]
diff --git a/Engine/Source/Editor/EditorWidgets/Private/SAssetSearchBox.cpp b/Engine/Source/Editor/EditorWidgets/Private/SAssetSearchBox.cpp
--- a/Engine/Source/Editor/EditorWidgets/Private/SAssetSearchBox.cpp
+++ b/Engine/Source/Editor/EditorWidgets/Private/SAssetSearchBox.cpp
@@ -200,7 +200,7 @@ TSharedRef<ITableRow> SAssetSearchBox::MakeSuggestionListItemWidget(TSharedPtr<F
 		SNew(STableRow< TSharedPtr<FString> >, OwnerTable)
 		[
 			SNew(STextBlock)
-			.Text(*Text.Get())
+			.Text(FText::FromString(*Text.Get()))
 			.HighlightText(this, &SAssetSearchBox::GetHighlightText)
 		];
 }
diff --git a/Engine/Source/Editor/EnvironmentQueryEditor/Private/DetailCustomizations/EnvDirectionCustomization.cpp b/Engine/Source/Editor/EnvironmentQueryEditor/Private/DetailCustomizations/EnvDirectionCustomization.cpp
--- a/Engine/Source/Editor/EnvironmentQueryEditor/Private/DetailCustomizations/EnvDirectionCustomization.cpp
+++ b/Engine/Source/Editor/EnvironmentQueryEditor/Private/DetailCustomizations/EnvDirectionCustomization.cpp
@@ -53,9 +53,9 @@ void FEnvDirectionCustomization::CustomizeChildren( TSharedRef<class IPropertyHa
 	.Visibility(TAttribute<EVisibility>::Create(TAttribute<EVisibility>::FGetter::CreateSP(this, &FEnvDirectionCustomization::GetRotationVisibility)));
 }
 
-FString FEnvDirectionCustomization::GetShortDescription() const
+FText FEnvDirectionCustomization::GetShortDescription() const
 {
-	return bIsRotation ? TEXT("context's rotation...") : TEXT("between two contexts...");
+	return bIsRotation ? LOCTEXT("DirectionShortDescRotatation", "context's rotation...") : LOCTEXT("DirectionShortDescBetweenTwoPoints", "between two contexts...");
 }
 
 EVisibility FEnvDirectionCustomization::GetTwoPointsVisibility() const
diff --git a/Engine/Source/Editor/EnvironmentQueryEditor/Private/DetailCustomizations/EnvDirectionCustomization.h b/Engine/Source/Editor/EnvironmentQueryEditor/Private/DetailCustomizations/EnvDirectionCustomization.h
--- a/Engine/Source/Editor/EnvironmentQueryEditor/Private/DetailCustomizations/EnvDirectionCustomization.h
+++ b/Engine/Source/Editor/EnvironmentQueryEditor/Private/DetailCustomizations/EnvDirectionCustomization.h
@@ -21,7 +21,7 @@ protected:
 	TSharedPtr<IPropertyHandle> ModeProp;
 	bool bIsRotation;
 
-	FString GetShortDescription() const;
+	FText GetShortDescription() const;
 	EVisibility GetTwoPointsVisibility() const;
 	EVisibility GetRotationVisibility() const;
 	void OnModeChanged();
diff --git a/Engine/Source/Editor/EnvironmentQueryEditor/Private/DetailCustomizations/EnvQueryParamSetupCustomization.cpp b/Engine/Source/Editor/EnvironmentQueryEditor/Private/DetailCustomizations/EnvQueryParamSetupCustomization.cpp
--- a/Engine/Source/Editor/EnvironmentQueryEditor/Private/DetailCustomizations/EnvQueryParamSetupCustomization.cpp
+++ b/Engine/Source/Editor/EnvironmentQueryEditor/Private/DetailCustomizations/EnvQueryParamSetupCustomization.cpp
@@ -178,16 +178,16 @@ void FEnvQueryParamSetupCustomization::OnParamBoolValueChanged(ECheckBoxState Bo
 	}
 }
 
-FString FEnvQueryParamSetupCustomization::GetComboText() const
+FText FEnvQueryParamSetupCustomization::GetComboText() const
 {
-	return bIsNamed ? FString("NAMED") : FString("VALUE");
+	return bIsNamed ? LOCTEXT("NamedLabel", "NAMED") : LOCTEXT("ValueLabel", "VALUE");
 }
 
-FString FEnvQueryParamSetupCustomization::GetComboTooltip() const
+FText FEnvQueryParamSetupCustomization::GetComboTooltip() const
 {
 	return bIsNamed ?
-		LOCTEXT("NamedTooltip","Value will be provided by named parameter, press to switch into direct mode.").ToString() :
-		LOCTEXT("ValueTooltip","Direct value of property, press to switch into parameter mode.").ToString();
+		LOCTEXT("NamedTooltip","Value will be provided by named parameter, press to switch into direct mode.") :
+		LOCTEXT("ValueTooltip","Direct value of property, press to switch into parameter mode.");
 }
 
 FReply FEnvQueryParamSetupCustomization::ToggleMode()
diff --git a/Engine/Source/Editor/EnvironmentQueryEditor/Private/DetailCustomizations/EnvQueryParamSetupCustomization.h b/Engine/Source/Editor/EnvironmentQueryEditor/Private/DetailCustomizations/EnvQueryParamSetupCustomization.h
--- a/Engine/Source/Editor/EnvironmentQueryEditor/Private/DetailCustomizations/EnvQueryParamSetupCustomization.h
+++ b/Engine/Source/Editor/EnvironmentQueryEditor/Private/DetailCustomizations/EnvQueryParamSetupCustomization.h
@@ -42,6 +42,6 @@ protected:
 
 	FReply ToggleMode();
 	void OnModeChanged();
-	FString GetComboText() const;
-	FString GetComboTooltip() const;
+	FText GetComboText() const;
+	FText GetComboTooltip() const;
 };
diff --git a/Engine/Source/Editor/EnvironmentQueryEditor/Private/DetailCustomizations/EnvQueryTestDetails.cpp b/Engine/Source/Editor/EnvironmentQueryEditor/Private/DetailCustomizations/EnvQueryTestDetails.cpp
--- a/Engine/Source/Editor/EnvironmentQueryEditor/Private/DetailCustomizations/EnvQueryTestDetails.cpp
+++ b/Engine/Source/Editor/EnvironmentQueryEditor/Private/DetailCustomizations/EnvQueryTestDetails.cpp
@@ -257,13 +257,13 @@ void FEnvQueryTestDetails::BuildFilterTestValues()
 	{
 		if (EditedTest->GetWorkOnFloatValues())
 		{
-			FilterTestValues.Add(FStringIntPair(TestConditionEnum->GetEnumText(EEnvTestFilterType::Minimum).ToString(), EEnvTestFilterType::Minimum));
-			FilterTestValues.Add(FStringIntPair(TestConditionEnum->GetEnumText(EEnvTestFilterType::Maximum).ToString(), EEnvTestFilterType::Maximum));
-			FilterTestValues.Add(FStringIntPair(TestConditionEnum->GetEnumText(EEnvTestFilterType::Range).ToString(), EEnvTestFilterType::Range));
+			FilterTestValues.Add(FTextIntPair(TestConditionEnum->GetEnumText(EEnvTestFilterType::Minimum), EEnvTestFilterType::Minimum));
+			FilterTestValues.Add(FTextIntPair(TestConditionEnum->GetEnumText(EEnvTestFilterType::Maximum), EEnvTestFilterType::Maximum));
+			FilterTestValues.Add(FTextIntPair(TestConditionEnum->GetEnumText(EEnvTestFilterType::Range), EEnvTestFilterType::Range));
 		}
 		else
 		{
-			FilterTestValues.Add(FStringIntPair(TestConditionEnum->GetEnumText(EEnvTestFilterType::Match).ToString(), EEnvTestFilterType::Match));
+			FilterTestValues.Add(FTextIntPair(TestConditionEnum->GetEnumText(EEnvTestFilterType::Match), EEnvTestFilterType::Match));
 		}
 	}
 }
@@ -276,16 +276,16 @@ void FEnvQueryTestDetails::BuildScoreEquationValues()
 	ScoreEquationValues.Reset();
 
 	// Const scoring is always valid.  But other equations are only valid if the score values can be other than boolean values.
-	ScoreEquationValues.Add(FStringIntPair(TestScoreEquationEnum->GetEnumText(EEnvTestScoreEquation::Constant).ToString(), EEnvTestScoreEquation::Constant));
+	ScoreEquationValues.Add(FTextIntPair(TestScoreEquationEnum->GetEnumText(EEnvTestScoreEquation::Constant), EEnvTestScoreEquation::Constant));
 
 	const UEnvQueryTest* EditedTest = Cast<const UEnvQueryTest>(MyTest.Get());
 	if (EditedTest)
 	{
 		if (EditedTest->GetWorkOnFloatValues())
 		{
-			ScoreEquationValues.Add(FStringIntPair(TestScoreEquationEnum->GetEnumText(EEnvTestScoreEquation::Linear).ToString(), EEnvTestScoreEquation::Linear));
-			ScoreEquationValues.Add(FStringIntPair(TestScoreEquationEnum->GetEnumText(EEnvTestScoreEquation::Square).ToString(), EEnvTestScoreEquation::Square));
-			ScoreEquationValues.Add(FStringIntPair(TestScoreEquationEnum->GetEnumText(EEnvTestScoreEquation::InverseLinear).ToString(), EEnvTestScoreEquation::InverseLinear));
+			ScoreEquationValues.Add(FTextIntPair(TestScoreEquationEnum->GetEnumText(EEnvTestScoreEquation::Linear), EEnvTestScoreEquation::Linear));
+			ScoreEquationValues.Add(FTextIntPair(TestScoreEquationEnum->GetEnumText(EEnvTestScoreEquation::Square), EEnvTestScoreEquation::Square));
+			ScoreEquationValues.Add(FTextIntPair(TestScoreEquationEnum->GetEnumText(EEnvTestScoreEquation::InverseLinear), EEnvTestScoreEquation::InverseLinear));
 		}
 	}
 }
@@ -302,14 +302,14 @@ void FEnvQueryTestDetails::OnScoreEquationChange(int32 Index)
 	ScoreEquationHandle->SetValue(EnumValue);
 }
 
-void FEnvQueryTestDetails::BuildScoreClampingTypeValues(bool bBuildMinValues, TArray<FStringIntPair>& ClampTypeValues) const
+void FEnvQueryTestDetails::BuildScoreClampingTypeValues(bool bBuildMinValues, TArray<FTextIntPair>& ClampTypeValues) const
 {
 	UEnum* ScoringNormalizationEnum = FindObject<UEnum>(ANY_PACKAGE, TEXT("EEnvQueryTestClamping"));
 	check(ScoringNormalizationEnum);
 
 	ClampTypeValues.Reset();
-	ClampTypeValues.Add(FStringIntPair(ScoringNormalizationEnum->GetEnumText(EEnvQueryTestClamping::None).ToString(), EEnvQueryTestClamping::None));
-	ClampTypeValues.Add(FStringIntPair(ScoringNormalizationEnum->GetEnumText(EEnvQueryTestClamping::SpecifiedValue).ToString(), EEnvQueryTestClamping::SpecifiedValue));
+	ClampTypeValues.Add(FTextIntPair(ScoringNormalizationEnum->GetEnumText(EEnvQueryTestClamping::None), EEnvQueryTestClamping::None));
+	ClampTypeValues.Add(FTextIntPair(ScoringNormalizationEnum->GetEnumText(EEnvQueryTestClamping::SpecifiedValue), EEnvQueryTestClamping::SpecifiedValue));
 
 	if (IsFiltering())
 	{
@@ -328,7 +328,7 @@ void FEnvQueryTestDetails::BuildScoreClampingTypeValues(bool bBuildMinValues, TA
 
 		if (bSupportFilterThreshold)
 		{
-			ClampTypeValues.Add(FStringIntPair(ScoringNormalizationEnum->GetEnumText(EEnvQueryTestClamping::FilterThreshold).ToString(), EEnvQueryTestClamping::FilterThreshold));
+			ClampTypeValues.Add(FTextIntPair(ScoringNormalizationEnum->GetEnumText(EEnvQueryTestClamping::FilterThreshold), EEnvQueryTestClamping::FilterThreshold));
 		}
 	}
 }
@@ -355,7 +355,7 @@ TSharedRef<SWidget> FEnvQueryTestDetails::OnGetClampMinTypeContent()
  	for (int32 i = 0; i < ClampMinTypeValues.Num(); i++)
  	{
  		FUIAction ItemAction(FExecuteAction::CreateSP(this, &FEnvQueryTestDetails::OnClampMinTestChange, ClampMinTypeValues[i].Int));
- 		MenuBuilder.AddMenuEntry(FText::FromString(ClampMinTypeValues[i].Str), TAttribute<FText>(), FSlateIcon(), ItemAction);
+ 		MenuBuilder.AddMenuEntry(ClampMinTypeValues[i].Text, TAttribute<FText>(), FSlateIcon(), ItemAction);
  	}
 
 	return MenuBuilder.MakeWidget();
@@ -369,13 +369,13 @@ TSharedRef<SWidget> FEnvQueryTestDetails::OnGetClampMaxTypeContent()
  	for (int32 i = 0; i < ClampMaxTypeValues.Num(); i++)
  	{
  		FUIAction ItemAction(FExecuteAction::CreateSP(this, &FEnvQueryTestDetails::OnClampMaxTestChange, ClampMaxTypeValues[i].Int));
- 		MenuBuilder.AddMenuEntry(FText::FromString(ClampMaxTypeValues[i].Str), TAttribute<FText>(), FSlateIcon(), ItemAction);
+ 		MenuBuilder.AddMenuEntry(ClampMaxTypeValues[i].Text, TAttribute<FText>(), FSlateIcon(), ItemAction);
  	}
 
 	return MenuBuilder.MakeWidget();
 }
 
-FString FEnvQueryTestDetails::GetClampMinTypeDesc() const
+FText FEnvQueryTestDetails::GetClampMinTypeDesc() const
 {
 	check(ClampMinTypeHandle.IsValid());
 	uint8 EnumValue;
@@ -385,14 +385,14 @@ FString FEnvQueryTestDetails::GetClampMinTypeDesc() const
 	{
 		if (ClampMinTypeValues[i].Int == EnumValue)
 		{
-			return ClampMinTypeValues[i].Str;
+			return ClampMinTypeValues[i].Text;
 		}
 	}
 
-	return FString();
+	return FText::GetEmpty();
 }
 
-FString FEnvQueryTestDetails::GetClampMaxTypeDesc() const
+FText FEnvQueryTestDetails::GetClampMaxTypeDesc() const
 {
 	check(ClampMaxTypeHandle.IsValid());
 	uint8 EnumValue;
@@ -402,14 +402,14 @@ FString FEnvQueryTestDetails::GetClampMaxTypeDesc() const
 	{
 		if (ClampMaxTypeValues[i].Int == EnumValue)
 		{
-			return ClampMaxTypeValues[i].Str;
+			return ClampMaxTypeValues[i].Text;
 		}
 	}
 
-	return FString();
+	return FText::GetEmpty();
 }
 
-FString FEnvQueryTestDetails::GetEquationValuesDesc() const
+FText FEnvQueryTestDetails::GetEquationValuesDesc() const
 {
 	check(ScoreEquationHandle.IsValid());
 	uint8 EnumValue;
@@ -419,11 +419,11 @@ FString FEnvQueryTestDetails::GetEquationValuesDesc() const
 	{
 		if (ScoreEquationValues[i].Int == EnumValue)
 		{
-			return ScoreEquationValues[i].Str;
+			return ScoreEquationValues[i].Text;
 		}
 	}
 
-	return FString();
+	return FText::GetEmpty();
 }
 
 TSharedRef<SWidget> FEnvQueryTestDetails::OnGetFilterTestContent()
@@ -434,13 +434,13 @@ TSharedRef<SWidget> FEnvQueryTestDetails::OnGetFilterTestContent()
 	for (int32 i = 0; i < FilterTestValues.Num(); i++)
 	{
 		FUIAction ItemAction(FExecuteAction::CreateSP(this, &FEnvQueryTestDetails::OnFilterTestChange, FilterTestValues[i].Int));
-		MenuBuilder.AddMenuEntry(FText::FromString(FilterTestValues[i].Str), TAttribute<FText>(), FSlateIcon(), ItemAction);
+		MenuBuilder.AddMenuEntry(FilterTestValues[i].Text, TAttribute<FText>(), FSlateIcon(), ItemAction);
 	}
 
 	return MenuBuilder.MakeWidget();
 }
 
-FString FEnvQueryTestDetails::GetCurrentFilterTestDesc() const
+FText FEnvQueryTestDetails::GetCurrentFilterTestDesc() const
 {
 	uint8 EnumValue;
 	FilterTypeHandle->GetValue(EnumValue);
@@ -449,11 +449,11 @@ FString FEnvQueryTestDetails::GetCurrentFilterTestDesc() const
 	{
 		if (FilterTestValues[i].Int == EnumValue)
 		{
-			return FilterTestValues[i].Str;
+			return FilterTestValues[i].Text;
 		}
 	}
 
-	return FString();
+	return FText::GetEmpty();
 }
 
 TSharedRef<SWidget> FEnvQueryTestDetails::OnGetEquationValuesContent()
@@ -464,13 +464,13 @@ TSharedRef<SWidget> FEnvQueryTestDetails::OnGetEquationValuesContent()
 	for (int32 i = 0; i < ScoreEquationValues.Num(); i++)
 	{
 		FUIAction ItemAction(FExecuteAction::CreateSP(this, &FEnvQueryTestDetails::OnScoreEquationChange, ScoreEquationValues[i].Int));
-		MenuBuilder.AddMenuEntry(FText::FromString(ScoreEquationValues[i].Str), TAttribute<FText>(), FSlateIcon(), ItemAction);
+		MenuBuilder.AddMenuEntry(ScoreEquationValues[i].Text, TAttribute<FText>(), FSlateIcon(), ItemAction);
 	}
 
 	return MenuBuilder.MakeWidget();
 }
 
-FString FEnvQueryTestDetails::GetScoreEquationInfo() const
+FText FEnvQueryTestDetails::GetScoreEquationInfo() const
 {
 	uint8 EnumValue;
 	ScoreEquationHandle->GetValue(EnumValue);
@@ -478,22 +478,22 @@ FString FEnvQueryTestDetails::GetScoreEquationInfo() const
 	switch (EnumValue)
 	{
 		case EEnvTestScoreEquation::Linear:
-			return LOCTEXT("Linear","Final score = ScoringFactor * NormalizedItemValue").ToString();
+			return LOCTEXT("Linear","Final score = ScoringFactor * NormalizedItemValue");
 
 		case EEnvTestScoreEquation::Square:
-			return LOCTEXT("Square","Final score = ScoringFactor * (NormalizedItemValue * NormalizedItemValue)\nBias towards items with big values.").ToString();
+			return LOCTEXT("Square","Final score = ScoringFactor * (NormalizedItemValue * NormalizedItemValue)\nBias towards items with big values.");
 
 		case EEnvTestScoreEquation::InverseLinear:
-			return LOCTEXT("Inverse","Final score = ScoringFactor * (1.0 - NormalizedItemValue)\nBias towards items with values close to zero.  (Linear, but flipped from 1 to 0 rather than 0 to 1.").ToString();
+			return LOCTEXT("Inverse","Final score = ScoringFactor * (1.0 - NormalizedItemValue)\nBias towards items with values close to zero.  (Linear, but flipped from 1 to 0 rather than 0 to 1.");
 
 		case EEnvTestScoreEquation::Constant:
-			return LOCTEXT("Constant", "Final score (for values that 'pass') = ScoringFactor\nNOTE: In this case, the score is normally EITHER the ScoringFactor value or zero.\nThe score will be zero if the Normalized Test Value is zero (or if the test value is false for a boolean query).\nOtherwise, score will be the ScoringFactor.").ToString();
+			return LOCTEXT("Constant", "Final score (for values that 'pass') = ScoringFactor\nNOTE: In this case, the score is normally EITHER the ScoringFactor value or zero.\nThe score will be zero if the Normalized Test Value is zero (or if the test value is false for a boolean query).\nOtherwise, score will be the ScoringFactor.");
 
 		default:
 			break;
 	}
 
-	return TEXT("");
+	return FText::GetEmpty();
 }
 
 EVisibility FEnvQueryTestDetails::GetVisibilityOfValueMinForScoreClamping() const
diff --git a/Engine/Source/Editor/EnvironmentQueryEditor/Private/DetailCustomizations/EnvQueryTestDetails.h b/Engine/Source/Editor/EnvironmentQueryEditor/Private/DetailCustomizations/EnvQueryTestDetails.h
--- a/Engine/Source/Editor/EnvironmentQueryEditor/Private/DetailCustomizations/EnvQueryTestDetails.h
+++ b/Engine/Source/Editor/EnvironmentQueryEditor/Private/DetailCustomizations/EnvQueryTestDetails.h
@@ -18,13 +18,13 @@ public:
 protected:
 
 	/** cached name, value pairs of enums that may change available options based on other options. */
-	struct FStringIntPair
+	struct FTextIntPair
 	{
-		FString Str;
+		FText Text;
 		int32 Int;
 
-		FStringIntPair() {}
-		FStringIntPair(FString InStr, int32 InInt) : Str(InStr), Int(InInt) {}
+		FTextIntPair() {}
+		FTextIntPair(FText InText, int32 InInt) : Text(InText), Int(InInt) {}
 	};
 
 	TSharedPtr<IPropertyHandle> FilterTypeHandle;
@@ -44,8 +44,8 @@ protected:
 	bool UsesFilterMin() const;
 	bool UsesFilterMax() const;
 
-	FString GetCurrentFilterTestDesc() const;
-	FString GetScoreEquationInfo() const;
+	FText GetCurrentFilterTestDesc() const;
+	FText GetScoreEquationInfo() const;
 
  	EVisibility GetScoreVisibility() const;
 	EVisibility GetClampingVisibility() const;
@@ -55,7 +55,7 @@ protected:
 
 	void BuildScoreEquationValues();
 	TSharedRef<SWidget> OnGetEquationValuesContent();
-	FString GetEquationValuesDesc() const;
+	FText GetEquationValuesDesc() const;
 	void OnScoreEquationChange(int32 Index);
 
 	void OnFilterTestChange(int32 Index);
@@ -63,10 +63,10 @@ protected:
 	void OnClampMaxTestChange(int32 Index);
 
 	TSharedRef<SWidget> OnGetClampMaxTypeContent();
-	FString GetClampMaxTypeDesc() const;
+	FText GetClampMaxTypeDesc() const;
 
 	TSharedRef<SWidget> OnGetClampMinTypeContent();
-	FString GetClampMinTypeDesc() const;
+	FText GetClampMinTypeDesc() const;
 
 	bool IsMatchingBoolValue() const;
 
@@ -89,19 +89,19 @@ protected:
 	EVisibility GetVisibilityOfScoreClampMinimum() const;
 	EVisibility GetVisibilityOfScoreClampMaximum() const;
 
-	void BuildScoreClampingTypeValues(bool bBuildMinValues, TArray<FStringIntPair>& ClampTypeValues) const;
+	void BuildScoreClampingTypeValues(bool bBuildMinValues, TArray<FTextIntPair>& ClampTypeValues) const;
 	
 	void UpdateTestFunctionPreview() const;
 	void FillEquationSamples(uint8 EquationType, bool bInversed, TArray<float>& Samples) const;
 	TSharedPtr<STestFunctionWidget> PreviewWidget;
 
 
-	TArray<FStringIntPair> FilterTestValues;
+	TArray<FTextIntPair> FilterTestValues;
 
-	TArray<FStringIntPair> ClampMinTypeValues;
-	TArray<FStringIntPair> ClampMaxTypeValues;
+	TArray<FTextIntPair> ClampMinTypeValues;
+	TArray<FTextIntPair> ClampMaxTypeValues;
 
-	TArray<FStringIntPair> ScoreEquationValues;
+	TArray<FTextIntPair> ScoreEquationValues;
 
 	TWeakObjectPtr<UObject> MyTest;
 
diff --git a/Engine/Source/Editor/EnvironmentQueryEditor/Private/DetailCustomizations/EnvTraceDataCustomization.cpp b/Engine/Source/Editor/EnvironmentQueryEditor/Private/DetailCustomizations/EnvTraceDataCustomization.cpp
--- a/Engine/Source/Editor/EnvironmentQueryEditor/Private/DetailCustomizations/EnvTraceDataCustomization.cpp
+++ b/Engine/Source/Editor/EnvironmentQueryEditor/Private/DetailCustomizations/EnvTraceDataCustomization.cpp
@@ -126,15 +126,15 @@ void FEnvTraceDataCustomization::CacheTraceModes(TSharedRef<class IPropertyHandl
 	TraceModes.Reset();
 	if (bCanDisable)
 	{
-		TraceModes.Add(FStringIntPair(TraceModeEnum->GetEnumText(EEnvQueryTrace::None).ToString(), EEnvQueryTrace::None));
+		TraceModes.Add(FTextIntPair(TraceModeEnum->GetEnumText(EEnvQueryTrace::None), EEnvQueryTrace::None));
 	}
 	if (bCanNavMesh)
 	{
-		TraceModes.Add(FStringIntPair(TraceModeEnum->GetEnumText(EEnvQueryTrace::Navigation).ToString(), EEnvQueryTrace::Navigation));
+		TraceModes.Add(FTextIntPair(TraceModeEnum->GetEnumText(EEnvQueryTrace::Navigation), EEnvQueryTrace::Navigation));
 	}
 	if (bCanGeometry)
 	{
-		TraceModes.Add(FStringIntPair(TraceModeEnum->GetEnumText(EEnvQueryTrace::Geometry).ToString(), EEnvQueryTrace::Geometry));
+		TraceModes.Add(FTextIntPair(TraceModeEnum->GetEnumText(EEnvQueryTrace::Geometry), EEnvQueryTrace::Geometry));
 	}
 
 	ActiveMode = EEnvQueryTrace::None;
@@ -154,41 +154,41 @@ TSharedRef<SWidget> FEnvTraceDataCustomization::OnGetTraceModeContent()
 	for (int32 i = 0; i < TraceModes.Num(); i++)
 	{
 		FUIAction ItemAction( FExecuteAction::CreateSP( this, &FEnvTraceDataCustomization::OnTraceModeChanged, TraceModes[i].Int ) );
-		MenuBuilder.AddMenuEntry( FText::FromString( TraceModes[i].Str ), TAttribute<FText>(), FSlateIcon(), ItemAction);
+		MenuBuilder.AddMenuEntry( TraceModes[i].Text, TAttribute<FText>(), FSlateIcon(), ItemAction);
 	}
 
 	return MenuBuilder.MakeWidget();
 }
 
-FString FEnvTraceDataCustomization::GetCurrentTraceModeDesc() const
+FText FEnvTraceDataCustomization::GetCurrentTraceModeDesc() const
 {
 	for (int32 i = 0; i < TraceModes.Num(); i++)
 	{
 		if (TraceModes[i].Int == ActiveMode)
 		{
-			return TraceModes[i].Str;
+			return TraceModes[i].Text;
 		}
 	}
 
-	return FString();
+	return FText::GetEmpty();
 }
 
-FString FEnvTraceDataCustomization::GetShortDescription() const
+FText FEnvTraceDataCustomization::GetShortDescription() const
 {
-	FString Desc;
+	FText Desc = FText::GetEmpty();
 
 	switch (ActiveMode)
 	{
 	case EEnvQueryTrace::Geometry:
-		Desc = LOCTEXT("TraceGeom","geometry trace").ToString();
+		Desc = LOCTEXT("TraceGeom","geometry trace");
 		break;
 
 	case EEnvQueryTrace::Navigation:
-		Desc = LOCTEXT("TraceNav","navmesh trace").ToString();
+		Desc = LOCTEXT("TraceNav","navmesh trace");
 		break;
 
 	case EEnvQueryTrace::None:
-		Desc = LOCTEXT("TraceNone","trace disabled").ToString();
+		Desc = LOCTEXT("TraceNone","trace disabled");
 		break;
 
 	default: break;
diff --git a/Engine/Source/Editor/EnvironmentQueryEditor/Private/DetailCustomizations/EnvTraceDataCustomization.h b/Engine/Source/Editor/EnvironmentQueryEditor/Private/DetailCustomizations/EnvTraceDataCustomization.h
--- a/Engine/Source/Editor/EnvironmentQueryEditor/Private/DetailCustomizations/EnvTraceDataCustomization.h
+++ b/Engine/Source/Editor/EnvironmentQueryEditor/Private/DetailCustomizations/EnvTraceDataCustomization.h
@@ -21,15 +21,15 @@ protected:
 	void CacheTraceModes(TSharedRef<class IPropertyHandle> StructPropertyHandle);
 
 	/** cached names of trace modes */
-	struct FStringIntPair
+	struct FTextIntPair
 	{
-		FString Str;
+		FText Text;
 		int32 Int;
 
-		FStringIntPair() {}
-		FStringIntPair(FString InStr, int32 InInt) : Str(InStr), Int(InInt) {}
+		FTextIntPair() {}
+		FTextIntPair(FText InText, int32 InInt) : Text(InText), Int(InInt) {}
 	};
-	TArray<FStringIntPair> TraceModes;
+	TArray<FTextIntPair> TraceModes;
 	bool bCanShowProjection;
 	uint8 ActiveMode;
 
@@ -38,8 +38,8 @@ protected:
 
 	void OnTraceModeChanged(int32 Index);
 	TSharedRef<SWidget> OnGetTraceModeContent();
-	FString GetCurrentTraceModeDesc() const;
-	FString GetShortDescription() const;
+	FText GetCurrentTraceModeDesc() const;
+	FText GetShortDescription() const;
 
 	EVisibility GetGeometryVisibility() const;
 	EVisibility GetNavigationVisibility() const;
diff --git a/Engine/Source/Editor/EnvironmentQueryEditor/Private/EnvironmentQueryEditor.cpp b/Engine/Source/Editor/EnvironmentQueryEditor/Private/EnvironmentQueryEditor.cpp
--- a/Engine/Source/Editor/EnvironmentQueryEditor/Private/EnvironmentQueryEditor.cpp
+++ b/Engine/Source/Editor/EnvironmentQueryEditor/Private/EnvironmentQueryEditor.cpp
@@ -181,7 +181,7 @@ TSharedRef<SGraphEditor> FEnvironmentQueryEditor::CreateGraphEditorWidget(UEdGra
 
 	// Create the appearance info
 	FGraphAppearanceInfo AppearanceInfo;
-	AppearanceInfo.CornerText = NSLOCTEXT("EnvironmentQueryEditor", "AppearanceCornerText", "ENVIRONMENT QUERY").ToString();
+	AppearanceInfo.CornerText = NSLOCTEXT("EnvironmentQueryEditor", "AppearanceCornerText", "ENVIRONMENT QUERY");
 
 	SGraphEditor::FGraphEditorEvents InEvents;
 	InEvents.OnSelectionChanged = SGraphEditor::FOnSelectionChanged::CreateSP(this, &FEnvironmentQueryEditor::OnSelectedNodesChanged);
diff --git a/Engine/Source/Editor/EnvironmentQueryEditor/Private/SGraphNode_EnvironmentQuery.cpp b/Engine/Source/Editor/EnvironmentQueryEditor/Private/SGraphNode_EnvironmentQuery.cpp
--- a/Engine/Source/Editor/EnvironmentQueryEditor/Private/SGraphNode_EnvironmentQuery.cpp
+++ b/Engine/Source/Editor/EnvironmentQueryEditor/Private/SGraphNode_EnvironmentQuery.cpp
@@ -667,10 +667,10 @@ TSharedPtr<SToolTip> SGraphNode_EnvironmentQuery::GetComplexTooltip()
 	return NULL;
 }
 
-FString SGraphNode_EnvironmentQuery::GetPreviewCornerText() const
+FText SGraphNode_EnvironmentQuery::GetPreviewCornerText() const
 {
 	UEnvironmentQueryGraphNode* StateNode = CastChecked<UEnvironmentQueryGraphNode>(GraphNode);
-	return FString(TEXT("Test CornerText"));
+	return FText::FromString(TEXT("Test CornerText"));
 }
 
 const FSlateBrush* SGraphNode_EnvironmentQuery::GetNameIcon() const
diff --git a/Engine/Source/Editor/EnvironmentQueryEditor/Private/SGraphNode_EnvironmentQuery.h b/Engine/Source/Editor/EnvironmentQueryEditor/Private/SGraphNode_EnvironmentQuery.h
--- a/Engine/Source/Editor/EnvironmentQueryEditor/Private/SGraphNode_EnvironmentQuery.h
+++ b/Engine/Source/Editor/EnvironmentQueryEditor/Private/SGraphNode_EnvironmentQuery.h
@@ -64,6 +64,6 @@ protected:
 	FSlateColor GetBackgroundColor() const;
 	FText GetDescription() const;
 
-	virtual FString GetPreviewCornerText() const;
+	virtual FText GetPreviewCornerText() const;
 	virtual const FSlateBrush* GetNameIcon() const;
 };
diff --git a/Engine/Source/Editor/GameProjectGeneration/Private/SNewClassDialog.cpp b/Engine/Source/Editor/GameProjectGeneration/Private/SNewClassDialog.cpp
--- a/Engine/Source/Editor/GameProjectGeneration/Private/SNewClassDialog.cpp
+++ b/Engine/Source/Editor/GameProjectGeneration/Private/SNewClassDialog.cpp
@@ -544,7 +544,7 @@ void SNewClassDialog::Construct( const FArguments& InArgs )
 									[
 										SNew(STextBlock)
 										.TextStyle( FEditorStyle::Get(), "NewClassDialog.SelectedParentClassLabel" )
-										.Text( LOCTEXT( "PathLabel", "Path" ).ToString() )
+										.Text( LOCTEXT( "PathLabel", "Path" ) )
 									]
 
 									// Path edit box
@@ -585,7 +585,7 @@ void SNewClassDialog::Construct( const FArguments& InArgs )
 									[
 										SNew(STextBlock)
 										.TextStyle( FEditorStyle::Get(), "NewClassDialog.SelectedParentClassLabel" )
-										.Text( LOCTEXT( "HeaderFileLabel", "Header File" ).ToString() )
+										.Text( LOCTEXT( "HeaderFileLabel", "Header File" ) )
 									]
 
 									// Header output text
@@ -609,7 +609,7 @@ void SNewClassDialog::Construct( const FArguments& InArgs )
 									[
 										SNew(STextBlock)
 										.TextStyle( FEditorStyle::Get(), "NewClassDialog.SelectedParentClassLabel" )
-										.Text( LOCTEXT( "SourceFileLabel", "Source File" ).ToString() )
+										.Text( LOCTEXT( "SourceFileLabel", "Source File" ) )
 									]
 
 									// Source output text
@@ -742,7 +742,7 @@ TSharedRef<ITableRow> SNewClassDialog::MakeParentClassListViewWidget(TSharedPtr<
 					[
 						SNew(STextBlock)
 						.TextStyle( FEditorStyle::Get(), "NewClassDialog.ParentClassItemTitle" )
-						.Text(ClassName)
+						.Text(FText::FromString(ClassName))
 					]
 				]
 
@@ -752,7 +752,7 @@ TSharedRef<ITableRow> SNewClassDialog::MakeParentClassListViewWidget(TSharedPtr<
 				[
 					SNew(STextBlock)
 					//.AutoWrapText(true)
-					.Text(ClassDescription)
+					.Text(FText::FromString(ClassDescription))
 				]
 			]
 		];
@@ -886,14 +886,14 @@ EVisibility SNewClassDialog::GetGlobalErrorLabelIDELinkVisibility() const
 	return FSourceCodeNavigation::IsCompilerAvailable() ? EVisibility::Collapsed : EVisibility::Visible;
 }
 
-FString SNewClassDialog::GetGlobalErrorLabelText() const
+FText SNewClassDialog::GetGlobalErrorLabelText() const
 {
 	if ( !FSourceCodeNavigation::IsCompilerAvailable() )
 	{
-		return FText::Format( LOCTEXT("NoCompilerFound", "No compiler was found. In order to use C++ code, you must first install {0}."), FSourceCodeNavigation::GetSuggestedSourceCodeIDE() ).ToString();
+		return FText::Format( LOCTEXT("NoCompilerFound", "No compiler was found. In order to use C++ code, you must first install {0}."), FSourceCodeNavigation::GetSuggestedSourceCodeIDE() );
 	}
 
-	return TEXT("");
+	return FText::GetEmpty();
 }
 
 void SNewClassDialog::OnNamePageEntered()
diff --git a/Engine/Source/Editor/GameProjectGeneration/Private/SNewClassDialog.h b/Engine/Source/Editor/GameProjectGeneration/Private/SNewClassDialog.h
--- a/Engine/Source/Editor/GameProjectGeneration/Private/SNewClassDialog.h
+++ b/Engine/Source/Editor/GameProjectGeneration/Private/SNewClassDialog.h
@@ -81,7 +81,7 @@ private:
 	EVisibility GetGlobalErrorLabelIDELinkVisibility() const;
 
 	/** Gets the text to display in the global error label */
-	FString GetGlobalErrorLabelText() const;
+	FText GetGlobalErrorLabelText() const;
 
 	/** Handler for when the user enters the "name class" page */
 	void OnNamePageEntered();
diff --git a/Engine/Source/Editor/GameProjectGeneration/Private/SNewProjectWizard.cpp b/Engine/Source/Editor/GameProjectGeneration/Private/SNewProjectWizard.cpp
--- a/Engine/Source/Editor/GameProjectGeneration/Private/SNewProjectWizard.cpp
+++ b/Engine/Source/Editor/GameProjectGeneration/Private/SNewProjectWizard.cpp
@@ -850,9 +850,9 @@ TSharedPtr<FTemplateItem> SNewProjectWizard::GetSelectedTemplateItem() const
 	return NULL;
 }
 
-FString SNewProjectWizard::GetSelectedTemplateClassTypes() const
+FText SNewProjectWizard::GetSelectedTemplateClassTypes() const
 {
-	return GetSelectedTemplateProperty<FString>(&FTemplateItem::ClassTypes);
+	return FText::FromString(GetSelectedTemplateProperty<FString>(&FTemplateItem::ClassTypes));
 }
 
 EVisibility SNewProjectWizard::GetSelectedTemplateClassVisibility() const
@@ -860,9 +860,9 @@ EVisibility SNewProjectWizard::GetSelectedTemplateClassVisibility() const
 	return GetSelectedTemplateProperty<FString>(&FTemplateItem::ClassTypes).IsEmpty() == false? EVisibility::Visible : EVisibility::Collapsed;
 }
 
-FString SNewProjectWizard::GetSelectedTemplateAssetTypes() const
+FText SNewProjectWizard::GetSelectedTemplateAssetTypes() const
 {
-	return GetSelectedTemplateProperty<FString>(&FTemplateItem::AssetTypes);
+	return FText::FromString(GetSelectedTemplateProperty<FString>(&FTemplateItem::AssetTypes));
 }
 
 EVisibility SNewProjectWizard::GetSelectedTemplateAssetVisibility() const
diff --git a/Engine/Source/Editor/GameProjectGeneration/Private/SNewProjectWizard.h b/Engine/Source/Editor/GameProjectGeneration/Private/SNewProjectWizard.h
--- a/Engine/Source/Editor/GameProjectGeneration/Private/SNewProjectWizard.h
+++ b/Engine/Source/Editor/GameProjectGeneration/Private/SNewProjectWizard.h
@@ -77,13 +77,13 @@ private:
 	EVisibility GetSelectedTemplatePreviewVisibility() const;
 
 	/** Get a string that details the class types referenced in the selected template */
-	FString GetSelectedTemplateClassTypes() const;
+	FText GetSelectedTemplateClassTypes() const;
 	
 	/** Get a visiblity of the class types display. If the string is empty this return Collapsed otherwise it will return Visible */
 	EVisibility GetSelectedTemplateClassVisibility() const;
 	
 	/** Get a string that details the class types referenced in the selected template */
-	FString GetSelectedTemplateAssetTypes() const;	
+	FText GetSelectedTemplateAssetTypes() const;	
 	
 	/** Get a visiblity of the asset types display. If the string is empty this return Collapsed otherwise it will return Visible */
 	EVisibility GetSelectedTemplateAssetVisibility() const;
diff --git a/Engine/Source/Editor/GameProjectGeneration/Private/SProjectBrowser.cpp b/Engine/Source/Editor/GameProjectGeneration/Private/SProjectBrowser.cpp
--- a/Engine/Source/Editor/GameProjectGeneration/Private/SProjectBrowser.cpp
+++ b/Engine/Source/Editor/GameProjectGeneration/Private/SProjectBrowser.cpp
@@ -442,7 +442,7 @@ TSharedRef<ITableRow> SProjectBrowser::MakeProjectViewWidget(TSharedPtr<FProject
 					.Padding(10)
 					[
 						SNew(STextBlock)
-						.Text(ProjectItem->GetEngineLabel())
+						.Text(FText::FromString(ProjectItem->GetEngineLabel()))
 						.TextStyle(FEditorStyle::Get(), "ProjectBrowser.VersionOverlayText")
 						.ColorAndOpacity(FLinearColor::White.CopyWithNewOpacity(0.5f))
 						.Visibility(ProjectItem->IsUpToDate() ? EVisibility::Collapsed : EVisibility::Visible)
diff --git a/Engine/Source/Editor/GameplayAbilitiesEditor/Private/AttributeDetails.cpp b/Engine/Source/Editor/GameplayAbilitiesEditor/Private/AttributeDetails.cpp
--- a/Engine/Source/Editor/GameplayAbilitiesEditor/Private/AttributeDetails.cpp
+++ b/Engine/Source/Editor/GameplayAbilitiesEditor/Private/AttributeDetails.cpp
@@ -540,28 +540,28 @@ TSharedRef<ITableRow> FScalableFloatDetails::HandleRowNameComboBoxGenarateWidget
 	return
 		SNew(STableRow<TSharedPtr<FString>>, OwnerTable)
 		[
-			SNew(STextBlock).Text(*InItem)
+			SNew(STextBlock).Text(FText::FromString(*InItem))
 		];
 }
 
 /** Display the current selection */
-FString FScalableFloatDetails::GetRowNameComboBoxContentText() const
+FText FScalableFloatDetails::GetRowNameComboBoxContentText() const
 {
-	FString RowName = TEXT("Multiple Values");
+	FString RowName;
 	const FPropertyAccess::Result RowResult = RowNameProperty->GetValue(RowName);
 	if (RowResult != FPropertyAccess::MultipleValues)
 	{
 		TSharedPtr<FString> SelectedRowName = CurrentSelectedItem;
 		if (SelectedRowName.IsValid())
 		{
-			RowName = *SelectedRowName;
+			return FText::FromString(*SelectedRowName);
 		}
 		else
 		{
-			RowName = TEXT("None");
+			return LOCTEXT("None", "None");
 		}
 	}
-	return RowName;
+	return LOCTEXT("MultipleValues", "Multiple Values");
 }
 
 FText FScalableFloatDetails::GetRowValuePreviewLabel() const
@@ -569,7 +569,7 @@ FText FScalableFloatDetails::GetRowValuePreviewLabel() const
 	return FText::Format(LOCTEXT("LevelPreviewLabel", "Preview At {0}"), FText::AsNumber(PreviewLevel));
 }
 
-FString FScalableFloatDetails::GetRowValuePreviewText() const
+FText FScalableFloatDetails::GetRowValuePreviewText() const
 {
 	TArray<const void*> RawPtrs;
 	CurveTableHandleProperty->AccessRawData(RawPtrs);
@@ -581,11 +581,14 @@ FString FScalableFloatDetails::GetRowValuePreviewText() const
 			float Value;
 			ValueProperty->GetValue(Value);
 
-			return FString::Printf(TEXT("%.3f"), Value * Curve.Eval(PreviewLevel));
+			static const FNumberFormattingOptions FormatOptions = FNumberFormattingOptions()
+				.SetMinimumFractionalDigits(3)
+				.SetMaximumFractionalDigits(3);
+			return FText::AsNumber(Value * Curve.Eval(PreviewLevel), &FormatOptions);
 		}
 	}
 
-	return FString();
+	return FText::GetEmpty();
 }
 
 /** Called by Slate when the filter box changes text. */
diff --git a/Engine/Source/Editor/GameplayAbilitiesEditor/Private/AttributeDetails.h b/Engine/Source/Editor/GameplayAbilitiesEditor/Private/AttributeDetails.h
--- a/Engine/Source/Editor/GameplayAbilitiesEditor/Private/AttributeDetails.h
+++ b/Engine/Source/Editor/GameplayAbilitiesEditor/Private/AttributeDetails.h
@@ -75,9 +75,9 @@ private:
 
 	TSharedRef<SWidget> GetListContent();
 	void OnSelectionChanged(TSharedPtr<FString> SelectedItem, ESelectInfo::Type SelectInfo);
-	FString GetRowNameComboBoxContentText() const;
+	FText GetRowNameComboBoxContentText() const;
 	FText GetRowValuePreviewLabel() const;
-	FString GetRowValuePreviewText() const;
+	FText GetRowValuePreviewText() const;
 	TSharedRef<ITableRow> HandleRowNameComboBoxGenarateWidget(TSharedPtr<FString> InItem, const TSharedRef<STableViewBase>& OwnerTable);
 	void OnFilterTextChanged(const FText& InFilterText);
 
diff --git a/Engine/Source/Editor/GameplayAbilitiesEditor/Private/SGraphNodeK2GameplayEffectVar.cpp b/Engine/Source/Editor/GameplayAbilitiesEditor/Private/SGraphNodeK2GameplayEffectVar.cpp
--- a/Engine/Source/Editor/GameplayAbilitiesEditor/Private/SGraphNodeK2GameplayEffectVar.cpp
+++ b/Engine/Source/Editor/GameplayAbilitiesEditor/Private/SGraphNodeK2GameplayEffectVar.cpp
@@ -38,7 +38,7 @@ void SGraphNodeK2GameplayEffectVar::UpdateGraphNode()
 	UGameplayEffect* GameplayEffect = GameplayEffectNode ? GameplayEffectNode->GameplayEffect : NULL;
 
 	FText DurationText = GameplayEffect ? GameplayEffect->DurationMagnitude.GetValueForEditorDisplay() : FText::GetEmpty();
-	FString PeriodString = FString::Printf(TEXT("%d s"), (GameplayEffect) ? GameplayEffect->Period.Value : 0.f);
+	FText PeriodText = FText::Format(LOCTEXT("PeriodTimeInSeconds", "{0} s"), FText::AsNumber((GameplayEffect) ? GameplayEffect->Period.Value : 0.f));
 
 	RightNodeBox->AddSlot()
 		[
@@ -47,7 +47,7 @@ void SGraphNodeK2GameplayEffectVar::UpdateGraphNode()
 			.HAlign(HAlign_Left)
 			[
 				SNew(STextBlock)
-				.Text(FString("Duration:"))
+				.Text(LOCTEXT("DurationLabel", "Duration:"))
 			]
 			+ SHorizontalBox::Slot()
 			.HAlign(HAlign_Right)
@@ -64,13 +64,13 @@ void SGraphNodeK2GameplayEffectVar::UpdateGraphNode()
 			.HAlign(HAlign_Left)
 			[
 				SNew(STextBlock)
-				.Text(FString("Period:"))
+				.Text(LOCTEXT("PeriodLabel", "Period:"))
 			]
 			+ SHorizontalBox::Slot()
 				.HAlign(HAlign_Right)
 				[
 					SNew(STextBlock)
-					.Text(PeriodString)
+					.Text(PeriodText)
 				]
 		];
 
@@ -82,8 +82,8 @@ void SGraphNodeK2GameplayEffectVar::UpdateGraphNode()
 
 	if (GameplayEffect && (GameplayEffect->ChanceToApplyToTarget.Value < 1.f || GameplayEffect->ChanceToExecuteOnGameplayEffect.Value < 1.f))
 	{
-		FString ApplyToTargetString = FString::Printf(TEXT("%d"), GameplayEffect->ChanceToApplyToTarget.Value);
-		FString ExecuteOnGEString = FString::Printf(TEXT("%d"), GameplayEffect->ChanceToExecuteOnGameplayEffect.Value);
+		FText ApplyToTargetText = FText::AsNumber(GameplayEffect->ChanceToApplyToTarget.Value);
+		FText ExecuteOnGEText = FText::AsNumber(GameplayEffect->ChanceToExecuteOnGameplayEffect.Value);
 
 		RightNodeBox->AddSlot()
 			[
@@ -92,13 +92,13 @@ void SGraphNodeK2GameplayEffectVar::UpdateGraphNode()
 				.HAlign(HAlign_Left)
 				[
 					SNew(STextBlock)
-					.Text(FString("Chance to Apply to Target:"))
+					.Text(LOCTEXT("ChanceToApplyToTargetLabel", "Chance to Apply to Target:"))
 				]
 				+ SHorizontalBox::Slot()
 					.HAlign(HAlign_Right)
 					[
 						SNew(STextBlock)
-						.Text(ApplyToTargetString)
+						.Text(ApplyToTargetText)
 					]
 			];
 
@@ -109,13 +109,13 @@ void SGraphNodeK2GameplayEffectVar::UpdateGraphNode()
 				.HAlign(HAlign_Left)
 				[
 					SNew(STextBlock)
-					.Text(FString("Chance to Execute on GE:"))
+					.Text(LOCTEXT("ChanceToExecuteOnGE", "Chance to Execute on GE:"))
 				]
 				+ SHorizontalBox::Slot()
 					.HAlign(HAlign_Right)
 					[
 						SNew(STextBlock)
-						.Text(ExecuteOnGEString)
+						.Text(ExecuteOnGEText)
 					]
 			];
 
@@ -140,10 +140,10 @@ void SGraphNodeK2GameplayEffectVar::UpdateGraphNode()
 			.HAlign(HAlign_Left)
 			[
 				SNew(STextBlock)
-				.Text(GameplayEffect->Modifiers[Idx].Attribute.GetName())
+				.Text(FText::FromString(GameplayEffect->Modifiers[Idx].Attribute.GetName()))
 			];
 
-		FString ModString = FString::Printf(TEXT("%s: %f"), *EGameplayModOpToString(GameplayEffect->Modifiers[Idx].ModifierOp), GameplayEffect->Modifiers[Idx].Magnitude.Value);
+		FText ModString = FText::Format(LOCTEXT("ModOpAndValue", "{0}: {1}"), FText::FromString(EGameplayModOpToString(GameplayEffect->Modifiers[Idx].ModifierOp)), FText::AsNumber(GameplayEffect->Modifiers[Idx].Magnitude.Value));
 
 		RightNodeBox->AddSlot()
 			.VAlign(VAlign_Center)
diff --git a/Engine/Source/Editor/GameplayTagsEditor/Private/GameplayTagContainerCustomization.cpp b/Engine/Source/Editor/GameplayTagsEditor/Private/GameplayTagContainerCustomization.cpp
--- a/Engine/Source/Editor/GameplayTagsEditor/Private/GameplayTagContainerCustomization.cpp
+++ b/Engine/Source/Editor/GameplayTagsEditor/Private/GameplayTagContainerCustomization.cpp
@@ -110,7 +110,7 @@ TSharedRef<ITableRow> FGameplayTagContainerCustomization::MakeListViewWidget(TSh
 {
 	return SNew( STableRow< TSharedPtr<FString> >, OwnerTable )
 			[
-				SNew(STextBlock) .Text( *Item.Get() )
+				SNew(STextBlock) .Text( FText::FromString(*Item.Get()) )
 			];
 }
 
diff --git a/Engine/Source/Editor/GameplayTagsEditor/Private/SGameplayTagContainerGraphPin.cpp b/Engine/Source/Editor/GameplayTagsEditor/Private/SGameplayTagContainerGraphPin.cpp
--- a/Engine/Source/Editor/GameplayTagsEditor/Private/SGameplayTagContainerGraphPin.cpp
+++ b/Engine/Source/Editor/GameplayTagsEditor/Private/SGameplayTagContainerGraphPin.cpp
@@ -127,7 +127,7 @@ TSharedRef<ITableRow> SGameplayTagContainerGraphPin::OnGenerateRow(TSharedPtr<FS
 {
 	return SNew( STableRow< TSharedPtr<FString> >, OwnerTable )
 		[
-			SNew(STextBlock) .Text( *Item.Get() )
+			SNew(STextBlock) .Text( FText::FromString(*Item.Get()) )
 		];
 }
 
diff --git a/Engine/Source/Editor/GraphEditor/Private/AnimationStateNodes/SGraphNodeAnimState.cpp b/Engine/Source/Editor/GraphEditor/Private/AnimationStateNodes/SGraphNodeAnimState.cpp
--- a/Engine/Source/Editor/GraphEditor/Private/AnimationStateNodes/SGraphNodeAnimState.cpp
+++ b/Engine/Source/Editor/GraphEditor/Private/AnimationStateNodes/SGraphNodeAnimState.cpp
@@ -313,11 +313,11 @@ TSharedPtr<SToolTip> SGraphNodeAnimState::GetComplexTooltip()
 		];
 }
 
-FString SGraphNodeAnimState::GetPreviewCornerText() const
+FText SGraphNodeAnimState::GetPreviewCornerText() const
 {
 	UAnimStateNodeBase* StateNode = CastChecked<UAnimStateNodeBase>(GraphNode);
 
-	return FString::Printf(TEXT("%s state"), *StateNode->GetStateName());
+	return FText::Format(NSLOCTEXT("SGraphNodeAnimState", "PreviewCornerStateText", "{0} state"), FText::FromString(StateNode->GetStateName()));
 }
 
 const FSlateBrush* SGraphNodeAnimState::GetNameIcon() const
@@ -338,11 +338,11 @@ void SGraphNodeAnimConduit::GetNodeInfoPopups(FNodeInfoContext* Context, TArray<
 	// Intentionally empty.
 }
 
-FString SGraphNodeAnimConduit::GetPreviewCornerText() const
+FText SGraphNodeAnimConduit::GetPreviewCornerText() const
 {
 	UAnimStateNodeBase* StateNode = CastChecked<UAnimStateNodeBase>(GraphNode);
 
-	return FString::Printf(TEXT("%s conduit"), *StateNode->GetStateName());
+	return FText::Format(NSLOCTEXT("SGraphNodeAnimState", "PreviewCornerConduitText", "{0} conduit"), FText::FromString(StateNode->GetStateName()));
 }
 
 const FSlateBrush* SGraphNodeAnimConduit::GetNameIcon() const
diff --git a/Engine/Source/Editor/GraphEditor/Private/AnimationStateNodes/SGraphNodeAnimState.h b/Engine/Source/Editor/GraphEditor/Private/AnimationStateNodes/SGraphNodeAnimState.h
--- a/Engine/Source/Editor/GraphEditor/Private/AnimationStateNodes/SGraphNodeAnimState.h
+++ b/Engine/Source/Editor/GraphEditor/Private/AnimationStateNodes/SGraphNodeAnimState.h
@@ -32,7 +32,7 @@ public:
 protected:
 	FSlateColor GetBorderBackgroundColor() const;
 
-	virtual FString GetPreviewCornerText() const;
+	virtual FText GetPreviewCornerText() const;
 	virtual const FSlateBrush* GetNameIcon() const;
 };
 
@@ -49,6 +49,6 @@ public:
 	virtual void GetNodeInfoPopups(FNodeInfoContext* Context, TArray<FGraphInformationPopupInfo>& Popups) const override;
 	// End of SNodePanel::SNode interface
 protected:
-	virtual FString GetPreviewCornerText() const override;
+	virtual FText GetPreviewCornerText() const override;
 	virtual const FSlateBrush* GetNameIcon() const override;
 };
