fixed test node reordering after rapid clicking in environment query editor
#ue4

diff --git a/Engine/Source/Editor/EnvironmentQueryEditor/Private/SGraphNode_EnvironmentQuery.cpp b/Engine/Source/Editor/EnvironmentQueryEditor/Private/SGraphNode_EnvironmentQuery.cpp
--- a/Engine/Source/Editor/EnvironmentQueryEditor/Private/SGraphNode_EnvironmentQuery.cpp
+++ b/Engine/Source/Editor/EnvironmentQueryEditor/Private/SGraphNode_EnvironmentQuery.cpp
@@ -16,6 +16,34 @@
 
 #define LOCTEXT_NAMESPACE "EnvironmentQueryEditor"
 
+TSharedRef<FDragNodeEQSTimed> FDragNodeEQSTimed::New(const TSharedRef<SGraphPanel>& InGraphPanel, const TSharedRef<SGraphNode>& InDraggedNode)
+{
+	TSharedRef<FDragNodeEQSTimed> Operation = MakeShareable(new FDragNodeEQSTimed);
+	Operation->StartTime = FPlatformTime::Seconds();
+	Operation->GraphPanel = InGraphPanel;
+	Operation->DraggedNodes.Add(InDraggedNode);
+	// adjust the decorator away from the current mouse location a small amount based on cursor size
+	Operation->DecoratorAdjust = FSlateApplication::Get().GetCursorSize();
+	Operation->Construct();
+	return Operation;
+}
+
+TSharedRef<FDragNodeEQSTimed> FDragNodeEQSTimed::New(const TSharedRef<SGraphPanel>& InGraphPanel, const TArray< TSharedRef<SGraphNode> >& InDraggedNodes)
+{
+	TSharedRef<FDragNodeEQSTimed> Operation = MakeShareable(new FDragNodeEQSTimed);
+	Operation->StartTime = FPlatformTime::Seconds();
+	Operation->GraphPanel = InGraphPanel;
+	Operation->DraggedNodes = InDraggedNodes;
+	Operation->DecoratorAdjust = FSlateApplication::Get().GetCursorSize();
+	Operation->Construct();
+	return Operation;
+}
+
+UEnvironmentQueryGraphNode_Test* FDragNodeEQSTimed::GetDropTargetNode() const
+{
+	return Cast<UEnvironmentQueryGraphNode_Test>(GetHoveredNode());
+}
+
 /////////////////////////////////////////////////////
 // SStateMachineOutputPin
 
@@ -84,7 +112,6 @@ void SGraphNode_EnvironmentQuery::Construct(const FArguments& InArgs, UEnvironme
 
 	this->UpdateGraphNode();
 
-	bIsMouseDown = false;
 	bDragMarkerVisible = false;
 }
 
@@ -196,7 +223,6 @@ void SGraphNode_EnvironmentQuery::UpdateGraphNode()
 			.Padding(0)
 			.BorderBackgroundColor( this, &SGraphNode_EnvironmentQuery::GetBorderBackgroundColor )
 			.OnMouseButtonDown(this, &SGraphNode_EnvironmentQuery::OnMouseDown)
-			.OnMouseButtonUp(this, &SGraphNode_EnvironmentQuery::OnMouseUp)
 			[
 				SNew(SOverlay)
 
@@ -401,24 +427,9 @@ void SGraphNode_EnvironmentQuery::AddPin(const TSharedRef<SGraphPin>& PinToAdd)
 	}
 }
 
-void SGraphNode_EnvironmentQuery::Tick( const FGeometry& AllottedGeometry, const double InCurrentTime, const float InDeltaTime )
-{
-	SGraphNode::Tick(AllottedGeometry, InCurrentTime, InDeltaTime);
-	if (bIsMouseDown)
-	{
-		MouseDownTime += InDeltaTime;
-	}
-}
-
 FReply SGraphNode_EnvironmentQuery::OnMouseMove(const FGeometry& SenderGeometry, const FPointerEvent& MouseEvent)
 {
-	if (!MouseEvent.IsMouseButtonDown(EKeys::LeftMouseButton) && bIsMouseDown)
-	{
-		bIsMouseDown = false;
-		MouseDownTime = 0;
-	}
-
-	if (bIsMouseDown && MouseDownTime > 0.1f)
+	if (MouseEvent.IsMouseButtonDown(EKeys::LeftMouseButton) && !(GEditor->bIsSimulatingInEditor || GEditor->PlayWorld))
 	{
 		//if we are holding mouse over a subnode
 		UEnvironmentQueryGraphNode_Test* TestNode = Cast<UEnvironmentQueryGraphNode_Test>(GraphNode);
@@ -426,23 +437,20 @@ FReply SGraphNode_EnvironmentQuery::OnMouseMove(const FGeometry& SenderGeometry,
 		{
 			const TSharedRef<SGraphPanel>& Panel = this->GetOwnerPanel().ToSharedRef();
 			const TSharedRef<SGraphNode>& Node = SharedThis(this);
-			return FReply::Handled().BeginDragDrop(FDragNode::New(Panel, Node));
+			return FReply::Handled().BeginDragDrop(FDragNodeEQSTimed::New(Panel, Node));
 		}
 	}
-	return FReply::Unhandled();
-}
 
-FReply SGraphNode_EnvironmentQuery::OnMouseUp(const FGeometry& SenderGeometry, const FPointerEvent& MouseEvent)
-{
-	bIsMouseDown = false;
-	MouseDownTime = 0;
+	if (!MouseEvent.IsMouseButtonDown(EKeys::LeftMouseButton) && bDragMarkerVisible)
+	{
+		SetDragMarker(false);
+	}
+
 	return FReply::Unhandled();
 }
 
 FReply SGraphNode_EnvironmentQuery::OnMouseDown(const FGeometry& SenderGeometry, const FPointerEvent& MouseEvent)
 {
-	bIsMouseDown = true;
-
 	UEnvironmentQueryGraphNode_Test* TestNode = Cast<UEnvironmentQueryGraphNode_Test>(GraphNode);
 	if (TestNode && TestNode->ParentNode)
 	{
@@ -582,14 +590,16 @@ void SGraphNode_EnvironmentQuery::OnDragLeave( const FDragDropEvent& DragDropEve
 		// Inform the Drag and Drop operation that we are not hovering any pins
 		DragConnectionOp->SetHoveredNode( TSharedPtr<SGraphNode>(NULL) );
 	}
-	SetDragMarker(false);
 
+	SetDragMarker(false);
 	SGraphNode::OnDragLeave(DragDropEvent);
 }
 
 FReply SGraphNode_EnvironmentQuery::OnDrop( const FGeometry& MyGeometry, const FDragDropEvent& DragDropEvent )
 {
-	TSharedPtr<FDragNode> DragNodeOp = DragDropEvent.GetOperationAs<FDragNode>();
+	SetDragMarker(false);
+
+	TSharedPtr<FDragNodeEQSTimed> DragNodeOp = DragDropEvent.GetOperationAs<FDragNodeEQSTimed>();
 	if (DragNodeOp.IsValid())
 	{
 		if (!DragNodeOp->IsValidOperation())
@@ -597,6 +607,12 @@ FReply SGraphNode_EnvironmentQuery::OnDrop( const FGeometry& MyGeometry, const F
 			return FReply::Handled();
 		}
 
+		const float DragTime = float(FPlatformTime::Seconds() - DragNodeOp->StartTime);
+		if (DragTime < 0.5f)
+		{
+			return FReply::Handled();
+		}
+
 		UEnvironmentQueryGraphNode_Option* MyNode = Cast<UEnvironmentQueryGraphNode_Option>(GraphNode);
 		if (MyNode == NULL)
 		{
@@ -607,9 +623,9 @@ FReply SGraphNode_EnvironmentQuery::OnDrop( const FGeometry& MyGeometry, const F
 		bool bReorderOperation = true;
 
 		const TArray< TSharedRef<SGraphNode> >& DraggedNodes = DragNodeOp->GetNodes();
-		for (int32 i = 0; i < DraggedNodes.Num(); i++)
+		for (int32 Idx = 0; Idx < DraggedNodes.Num(); Idx++)
 		{
-			UEnvironmentQueryGraphNode_Test* DraggedTestNode = Cast<UEnvironmentQueryGraphNode_Test>(DraggedNodes[i]->GetNodeObj());
+			UEnvironmentQueryGraphNode_Test* DraggedTestNode = Cast<UEnvironmentQueryGraphNode_Test>(DraggedNodes[Idx]->GetNodeObj());
 			if (DraggedTestNode && DraggedTestNode->ParentNode)
 			{
 				if (DraggedTestNode->ParentNode != GraphNode)
@@ -619,34 +635,34 @@ FReply SGraphNode_EnvironmentQuery::OnDrop( const FGeometry& MyGeometry, const F
 
 				DraggedTestNode->ParentNode->Modify();
 				TArray<UEnvironmentQueryGraphNode_Test*> & SubNodes = DraggedTestNode->ParentNode->Tests;
-				for (int32 j = 0; j < SubNodes.Num(); j++)
+				for (int32 SubIdx = 0; SubIdx < SubNodes.Num(); SubIdx++)
 				{
-					if (SubNodes[j] == DraggedTestNode)
+					if (SubNodes[SubIdx] == DraggedTestNode)
 					{
-						SubNodes.RemoveAt(j);
+						SubNodes.RemoveAt(SubIdx);
 					}
 				}
 			}
 		}
 
-		int32 InsertIndex = 0;
-		for (int32 i=0; i < TestWidgets.Num(); i++)
-		{
-			TSharedPtr<SGraphNode_EnvironmentQuery> TestNodeWidget = StaticCastSharedPtr<SGraphNode_EnvironmentQuery>(TestWidgets[i]);
-			if (TestNodeWidget->bDragMarkerVisible)
-			{
-				InsertIndex = MyNode->Tests.IndexOfByKey(TestNodeWidget->GetNodeObj());
-				break;
-			}
-		}
+		UEnvironmentQueryGraphNode_Test* DropTargetNode = DragNodeOp->GetDropTargetNode();
+		const int32 InsertIndex = MyNode->Tests.IndexOfByKey(DropTargetNode);
 
-		for (int32 i = 0; i < DraggedNodes.Num(); i++)
+		for (int32 Idx = 0; Idx < DraggedNodes.Num(); Idx++)
 		{
-			UEnvironmentQueryGraphNode_Test* DraggedTestNode = Cast<UEnvironmentQueryGraphNode_Test>(DraggedNodes[i]->GetNodeObj());
+			UEnvironmentQueryGraphNode_Test* DraggedTestNode = Cast<UEnvironmentQueryGraphNode_Test>(DraggedNodes[Idx]->GetNodeObj());
 			DraggedTestNode->Modify();
 			DraggedTestNode->ParentNode = MyNode;
 			MyNode->Modify();
-			MyNode->Tests.Insert(DraggedTestNode, InsertIndex);
+
+			if (InsertIndex > -1)
+			{
+				MyNode->Tests.Insert(DraggedTestNode, InsertIndex);
+			}
+			else
+			{
+				MyNode->Tests.Add(DraggedTestNode);
+			}
 		}
 
 		if (bReorderOperation)
diff --git a/Engine/Source/Editor/EnvironmentQueryEditor/Private/SGraphNode_EnvironmentQuery.h b/Engine/Source/Editor/EnvironmentQueryEditor/Private/SGraphNode_EnvironmentQuery.h
--- a/Engine/Source/Editor/EnvironmentQueryEditor/Private/SGraphNode_EnvironmentQuery.h
+++ b/Engine/Source/Editor/EnvironmentQueryEditor/Private/SGraphNode_EnvironmentQuery.h
@@ -2,6 +2,24 @@
 
 #pragma once
 
+class UEnvironmentQueryGraphNode_Test;
+
+class FDragNodeEQSTimed : public FDragNode
+{
+public:
+	DRAG_DROP_OPERATOR_TYPE(FDragNodeEQSTimed, FDragNode)
+
+	static TSharedRef<FDragNodeEQSTimed> New(const TSharedRef<SGraphPanel>& InGraphPanel, const TSharedRef<SGraphNode>& InDraggedNode);
+	static TSharedRef<FDragNodeEQSTimed> New(const TSharedRef<SGraphPanel>& InGraphPanel, const TArray< TSharedRef<SGraphNode> >& InDraggedNodes);
+
+	UEnvironmentQueryGraphNode_Test* GetDropTargetNode() const;
+
+	double StartTime;
+
+protected:
+	typedef FDragNode Super;
+};
+
 class SGraphNode_EnvironmentQuery : public SGraphNode
 {
 public:
@@ -20,7 +38,6 @@ public:
 	virtual FReply OnDrop( const FGeometry& MyGeometry, const FDragDropEvent& DragDropEvent ) override;
 	virtual void OnDragLeave( const FDragDropEvent& DragDropEvent ) override;
 	virtual FReply OnMouseMove(const FGeometry& SenderGeometry, const FPointerEvent& MouseEvent) override;
-	virtual void Tick( const FGeometry& AllottedGeometry, const double InCurrentTime, const float InDeltaTime ) override;
 	virtual void SetOwner(const TSharedRef<SGraphPanel>& OwnerPanel) override;
 
 	// End of SGraphNode interface
@@ -28,9 +45,6 @@ public:
 	/** handle mouse down on the node */
 	FReply OnMouseDown(const FGeometry& SenderGeometry, const FPointerEvent& MouseEvent);
 
-	/** handle mouse up on the node */
-	FReply OnMouseUp(const FGeometry& SenderGeometry, const FPointerEvent& MouseEvent);
-
 	/** adds decorator widget inside current node */
 	void AddTest(TSharedPtr<SGraphNode> TestWidget);
 
@@ -55,10 +69,7 @@ protected:
 	TArray< TSharedPtr<SGraphNode> > TestWidgets;
 	TSharedPtr<SVerticalBox> TestBox;
 
-	float MouseDownTime;
-
 	uint32 bDragMarkerVisible : 1;
-	uint32 bIsMouseDown : 1;
 
 	FSlateColor GetBorderBackgroundColor() const;
 	FSlateColor GetBackgroundColor() const;
