Added a very simple mechanism to have EQS editor include new BP-implemented generators in generators list #UE4

And not crash when one gets removed!
4.7 critical fix.
This is to be fixed properly in time for 4.8, to reuse the solution BT editor has.

diff --git a/Engine/Source/Editor/EnvironmentQueryEditor/Private/EdGraphSchema_EnvironmentQuery.cpp b/Engine/Source/Editor/EnvironmentQueryEditor/Private/EdGraphSchema_EnvironmentQuery.cpp
--- a/Engine/Source/Editor/EnvironmentQueryEditor/Private/EdGraphSchema_EnvironmentQuery.cpp
+++ b/Engine/Source/Editor/EnvironmentQueryEditor/Private/EdGraphSchema_EnvironmentQuery.cpp
@@ -10,6 +10,9 @@
 #include "Kismet2/KismetEditorUtilities.h"
 #include "GenericCommands.h"
 
+#include "AssetRegistryModule.h"
+#include "HotReloadInterface.h"
+
 #define LOCTEXT_NAMESPACE "EnvironmentQuerySchema"
 #define SNAP_GRID (16) // @todo ensure this is the same as SNodePanel::GetSnapGridSize()
 
@@ -118,6 +121,52 @@ void FEnvironmentQuerySchemaAction_NewSubNode::AddReferencedObjects( FReferenceC
 
 //////////////////////////////////////////////////////////////////////////
 
+namespace FEQSClassCacheInfo
+{
+	static bool bEQSGeneratorCacheInvalid = true;
+	static bool bEQSTestCacheInvalid = true;
+	void InvalidateCache()
+	{
+		bEQSGeneratorCacheInvalid = true;
+		bEQSTestCacheInvalid = true;
+	}
+
+	void InvalidateCache(const class FAssetData& AssetData)
+	{
+		bEQSGeneratorCacheInvalid = true;
+		bEQSTestCacheInvalid = true;
+	}
+
+	void InvalidateCache(bool bWasTriggeredAutomatically)
+	{
+		bEQSGeneratorCacheInvalid = true;
+		bEQSTestCacheInvalid = true;
+	}
+
+	void SetupCacheObservers()
+	{
+		static bool bCacheFixDelegatesSetUp = false;
+		if (bCacheFixDelegatesSetUp == false)
+		{
+			bCacheFixDelegatesSetUp = true;
+
+			// Register with the Asset Registry to be informed when it is done loading up files.
+			FAssetRegistryModule& AssetRegistryModule = FModuleManager::GetModuleChecked<FAssetRegistryModule>(TEXT("AssetRegistry"));
+			AssetRegistryModule.Get().OnFilesLoaded().AddStatic(&FEQSClassCacheInfo::InvalidateCache);
+			AssetRegistryModule.Get().OnAssetAdded().AddStatic(&FEQSClassCacheInfo::InvalidateCache);
+			AssetRegistryModule.Get().OnAssetRemoved().AddStatic(&FEQSClassCacheInfo::InvalidateCache);
+
+			// Register to have Populate called when doing a Hot Reload.
+			IHotReloadInterface& HotReloadSupport = FModuleManager::LoadModuleChecked<IHotReloadInterface>("HotReload");
+			HotReloadSupport.OnHotReload().AddStatic(&FEQSClassCacheInfo::InvalidateCache);
+
+			// Register to have Populate called when a Blueprint is compiled.
+			GEditor->OnBlueprintCompiled().AddStatic(&FEQSClassCacheInfo::InvalidateCache);
+			GEditor->OnClassPackageLoadedOrUnloaded().AddStatic(&FEQSClassCacheInfo::InvalidateCache);
+		}
+	}
+}
+
 UEdGraphSchema_EnvironmentQuery::UEdGraphSchema_EnvironmentQuery(const FObjectInitializer& ObjectInitializer)
 	: Super(ObjectInitializer)
 {
@@ -171,8 +220,12 @@ FString GetNodeDescriptionHelper(UObject* Ob)
 static void GatherEQSGenerators(TArray<UClass*>& Classes)
 {
 	static TArray<UClass*> CachedClasses;
-	if (CachedClasses.Num() == 0)
+	if (CachedClasses.Num() == 0 || FEQSClassCacheInfo::bEQSGeneratorCacheInvalid)
 	{
+		FEQSClassCacheInfo::SetupCacheObservers();
+		FEQSClassCacheInfo::bEQSGeneratorCacheInvalid = false;
+		CachedClasses.Reset();
+
 		for (TObjectIterator<UClass> It; It; ++It)
 		{
 			UClass* TestClass = *It;
@@ -195,8 +248,12 @@ static void GatherEQSGenerators(TArray<UClass*>& Classes)
 static void GatherEQSTests(TArray<UClass*>& Classes)
 {
 	static TArray<UClass*> CachedClasses;
-	if (CachedClasses.Num() == 0)
+	if (CachedClasses.Num() == 0 || FEQSClassCacheInfo::bEQSTestCacheInvalid)
 	{
+		FEQSClassCacheInfo::SetupCacheObservers();
+		FEQSClassCacheInfo::bEQSTestCacheInvalid = false;
+		CachedClasses.Reset();
+
 		for (TObjectIterator<UClass> It; It; ++It)
 		{
 			UClass* TestClass = *It;
