unified shared parts of AI editors (nested nodes, copy/paste/drag operations)

diff --git a/Engine/Source/Editor/EnvironmentQueryEditor/Classes/EdGraphSchema_EnvironmentQuery.h b/Engine/Source/Editor/EnvironmentQueryEditor/Classes/EdGraphSchema_EnvironmentQuery.h
--- a/Engine/Source/Editor/EnvironmentQueryEditor/Classes/EdGraphSchema_EnvironmentQuery.h
+++ b/Engine/Source/Editor/EnvironmentQueryEditor/Classes/EdGraphSchema_EnvironmentQuery.h
@@ -2,94 +2,22 @@
 
 #pragma once
 
+#include "AIGraphSchema.h"
 #include "EdGraphSchema_EnvironmentQuery.generated.h"
 
-/** Action to add a node to the graph */
-USTRUCT()
-struct FEnvironmentQuerySchemaAction_NewNode : public FEdGraphSchemaAction
-{
-	GENERATED_USTRUCT_BODY();
-
-	/** Template of node we want to create */
-	UPROPERTY()
-	class UEnvironmentQueryGraphNode* NodeTemplate;
-
-
-	FEnvironmentQuerySchemaAction_NewNode() 
-		: FEdGraphSchemaAction()
-		, NodeTemplate(NULL)
-	{}
-
-	FEnvironmentQuerySchemaAction_NewNode(const FString& InNodeCategory, const FText& InMenuDesc, const FString& InToolTip, const int32 InGrouping)
-		: FEdGraphSchemaAction(InNodeCategory, InMenuDesc, InToolTip, InGrouping) 
-		, NodeTemplate(NULL)
-	{}
-
-	// FEdGraphSchemaAction interface
-	virtual UEdGraphNode* PerformAction(class UEdGraph* ParentGraph, UEdGraphPin* FromPin, const FVector2D Location, bool bSelectNewNode = true) override;
-	virtual UEdGraphNode* PerformAction(class UEdGraph* ParentGraph, TArray<UEdGraphPin*>& FromPins, const FVector2D Location, bool bSelectNewNode = true) override;
-	virtual void AddReferencedObjects( FReferenceCollector& Collector ) override;
-	// End of FEdGraphSchemaAction interface
-
-	template <typename NodeType>
-	static NodeType* SpawnNodeFromTemplate(class UEdGraph* ParentGraph, NodeType* InTemplateNode, const FVector2D Location)
-	{
-		FEnvironmentQuerySchemaAction_NewNode Action;
-		Action.NodeTemplate = InTemplateNode;
-
-		return Cast<NodeType>(Action.PerformAction(ParentGraph, NULL, Location));
-	}
-};
-
-/** Action to add a subnode to the selected node */
-USTRUCT()
-struct FEnvironmentQuerySchemaAction_NewSubNode : public FEdGraphSchemaAction
-{
-	GENERATED_USTRUCT_BODY();
-
-	/** Template of node we want to create */
-	UPROPERTY()
-	class UEnvironmentQueryGraphNode_Test* NodeTemplate;
-
-	/** parent node */
-	UPROPERTY()
-	class UEnvironmentQueryGraphNode_Option* ParentNode;
-
-	FEnvironmentQuerySchemaAction_NewSubNode() 
-		: FEdGraphSchemaAction()
-		, NodeTemplate(NULL)
-	{}
-
-	FEnvironmentQuerySchemaAction_NewSubNode(const FString& InNodeCategory, const FText& InMenuDesc, const FString& InToolTip, const int32 InGrouping)
-		: FEdGraphSchemaAction(InNodeCategory, InMenuDesc, InToolTip, InGrouping) 
-		, NodeTemplate(NULL)
-	{}
-
-	// FEdGraphSchemaAction interface
-	virtual UEdGraphNode* PerformAction(class UEdGraph* ParentGraph, UEdGraphPin* FromPin, const FVector2D Location, bool bSelectNewNode = true) override;
-	virtual UEdGraphNode* PerformAction(class UEdGraph* ParentGraph, TArray<UEdGraphPin*>& FromPins, const FVector2D Location, bool bSelectNewNode = true) override;
-	virtual void AddReferencedObjects( FReferenceCollector& Collector ) override;
-	// End of FEdGraphSchemaAction interface
-};
-
 UCLASS(MinimalAPI)
-class UEdGraphSchema_EnvironmentQuery : public UEdGraphSchema
+class UEdGraphSchema_EnvironmentQuery : public UAIGraphSchema
 {
 	GENERATED_UCLASS_BODY()
 
-	void GetBreakLinkToSubMenuActions(class FMenuBuilder& MenuBuilder, class UEdGraphPin* InGraphPin);
-
 	// Begin EdGraphSchema interface
 	virtual void CreateDefaultNodesForGraph(UEdGraph& Graph) const override;
 	virtual void GetGraphContextActions(FGraphContextMenuBuilder& ContextMenuBuilder) const override;
-	virtual void GetContextMenuActions(const UEdGraph* CurrentGraph, const UEdGraphNode* InGraphNode, const UEdGraphPin* InGraphPin, class FMenuBuilder* MenuBuilder, bool bIsDebugging) const override;
 	virtual const FPinConnectionResponse CanCreateConnection(const UEdGraphPin* A, const UEdGraphPin* B) const override;
 	virtual const FPinConnectionResponse CanMergeNodes(const UEdGraphNode* A, const UEdGraphNode* B) const override;
-	virtual FLinearColor GetPinTypeColor(const FEdGraphPinType& PinType) const override;
-	virtual bool ShouldHidePinDefaultValue(UEdGraphPin* Pin) const override;
-	virtual class FConnectionDrawingPolicy* CreateConnectionDrawingPolicy(int32 InBackLayerID, int32 InFrontLayerID, float InZoomFactor, const FSlateRect& InClippingRect, class FSlateWindowElementList& InDrawElements, class UEdGraph* InGraphObj) const override;
 	// End EdGraphSchema interface
 
-	void GetGraphNodeContextActions(FGraphContextMenuBuilder& ContextMenuBuilder) const;
+	// Begin UAIGraphSchema interface
+	virtual void GetSubNodeClasses(int32 SubNodeFlags, TArray<FGraphNodeClassData>& ClassData, UClass*& GraphNodeClass) const;
+	// End UAIGraphSchema interface
 };
-
diff --git a/Engine/Source/Editor/EnvironmentQueryEditor/Classes/EnvironmentQueryGraph.h b/Engine/Source/Editor/EnvironmentQueryEditor/Classes/EnvironmentQueryGraph.h
--- a/Engine/Source/Editor/EnvironmentQueryEditor/Classes/EnvironmentQueryGraph.h
+++ b/Engine/Source/Editor/EnvironmentQueryEditor/Classes/EnvironmentQueryGraph.h
@@ -1,33 +1,28 @@
 // Copyright 1998-2015 Epic Games, Inc. All Rights Reserved.
 #pragma once
+
+#include "AIGraph.h"
 #include "EnvironmentQueryGraph.generated.h"
 
 UCLASS()
-class UEnvironmentQueryGraph : public UEdGraph
+class UEnvironmentQueryGraph : public UAIGraph
 {
 	GENERATED_UCLASS_BODY()
 
-	UPROPERTY()
-	int32 GraphVersion;
-
-	void UpdateVersion();
-	void MarkVersion();
+	virtual void UpdateVersion() override;
+	virtual void MarkVersion() override;
+	virtual void UpdateAsset(int32 UpdateFlags = 0) override;
 
 	void CalculateAllWeights();
-	void UpdateAsset();
-	void RemoveOrphanedNodes();
 	void CreateEnvQueryFromGraph(class UEnvironmentQueryGraphNode* RootEdNode);
 
-	void LockUpdates();
-	void UnlockUpdates();
-
 protected:
 
-	/** if set, graph modifications won't cause updates in internal tree structure (skipping UpdateAsset)
-	 * flag allows freezing update during heavy changes like pasting new nodes */
-	uint32 bLockUpdates : 1;
-
 	void UpdateVersion_NestedNodes();
 	void UpdateVersion_FixupOuters();
-};
+	void UpdateVersion_CollectClassData();
+
+	void UpdateNodeClassData(UAIGraphNode* UpdateNode, UClass* InstanceClass);
 
+	virtual void CollectAllNodeInstances(TSet<UObject*>& NodeInstances) override;
+};
diff --git a/Engine/Source/Editor/EnvironmentQueryEditor/Classes/EnvironmentQueryGraphNode.h b/Engine/Source/Editor/EnvironmentQueryEditor/Classes/EnvironmentQueryGraphNode.h
--- a/Engine/Source/Editor/EnvironmentQueryEditor/Classes/EnvironmentQueryGraphNode.h
+++ b/Engine/Source/Editor/EnvironmentQueryEditor/Classes/EnvironmentQueryGraphNode.h
@@ -1,48 +1,19 @@
 // Copyright 1998-2015 Epic Games, Inc. All Rights Reserved.
 
 #pragma once
+
+#include "AIGraphNode.h"
 #include "EnvironmentQueryGraphNode.generated.h"
 
 UCLASS()
-class UEnvironmentQueryGraphNode : public UEdGraphNode
+class UEnvironmentQueryGraphNode : public UAIGraphNode
 {
 	GENERATED_UCLASS_BODY()
 
-	UPROPERTY()
-	UObject* NodeInstance;
-
-	/** subnode index assigned during copy operation to connect nodes again on paste */
-	UPROPERTY()
-	int32 CopySubNodeIndex;
-
-	/** Get the BT graph that owns this node */
 	virtual class UEnvironmentQueryGraph* GetEnvironmentQueryGraph();
 
-	virtual void AutowireNewNode(UEdGraphPin* FromPin) override;
-	virtual void PostEditImport() override;
-	virtual void PrepareForCopying() override;
-	virtual void PostCopyNode();
-
-	// @return the input pin for this state
-	virtual UEdGraphPin* GetInputPin(int32 InputIndex=0) const;
-	// @return the output pin for this state
-	virtual UEdGraphPin* GetOutputPin(int32 InputIndex=0) const;
-	//
-	virtual UEdGraph* GetBoundGraph() const { return NULL; }
-
 	virtual FText GetNodeTitle(ENodeTitleType::Type TitleType) const override;
-	virtual FText GetDescription() const;
-
-	virtual void NodeConnectionListChanged() override;
+	virtual FText GetDescription() const override;
 
 	virtual bool CanCreateUnderSpecifiedSchema(const UEdGraphSchema* DesiredSchema) const override;
-
-	virtual bool IsSubNode() const { return false; }
-
-	UClass* EnvQueryNodeClass;
-
-protected:
-
-	virtual void ResetNodeOwner();
 };
-
diff --git a/Engine/Source/Editor/EnvironmentQueryEditor/Classes/EnvironmentQueryGraphNode_Option.h b/Engine/Source/Editor/EnvironmentQueryEditor/Classes/EnvironmentQueryGraphNode_Option.h
--- a/Engine/Source/Editor/EnvironmentQueryEditor/Classes/EnvironmentQueryGraphNode_Option.h
+++ b/Engine/Source/Editor/EnvironmentQueryEditor/Classes/EnvironmentQueryGraphNode_Option.h
@@ -8,17 +8,14 @@ class UEnvironmentQueryGraphNode_Option : public UEnvironmentQueryGraphNode
 {
 	GENERATED_UCLASS_BODY()
 
-	UPROPERTY()
-	TArray<class UEnvironmentQueryGraphNode_Test*> Tests;
-
 	virtual void AllocateDefaultPins() override;
 	virtual void PostPlacedNewNode() override;
-	virtual void GetContextMenuActions(const FGraphNodeContextMenuBuilder& Context) const override;
-	virtual void PrepareForCopying() override;
 	virtual FText GetNodeTitle(ENodeTitleType::Type TitleType) const override;
 	virtual FText GetDescription() const override;
 
-	void AddSubNode(class UEnvironmentQueryGraphNode_Test* NodeTemplate, class UEdGraph* ParentGraph);
+	virtual void PrepareForCopying() override;
+
+	virtual void GetContextMenuActions(const FGraphNodeContextMenuBuilder& Context) const override;
 	void CreateAddTestSubMenu(class FMenuBuilder& MenuBuilder, UEdGraph* Graph) const;
 
 	void CalculateWeights();
diff --git a/Engine/Source/Editor/EnvironmentQueryEditor/Classes/EnvironmentQueryGraphNode_Root.h b/Engine/Source/Editor/EnvironmentQueryEditor/Classes/EnvironmentQueryGraphNode_Root.h
--- a/Engine/Source/Editor/EnvironmentQueryEditor/Classes/EnvironmentQueryGraphNode_Root.h
+++ b/Engine/Source/Editor/EnvironmentQueryEditor/Classes/EnvironmentQueryGraphNode_Root.h
@@ -8,6 +8,5 @@ class UEnvironmentQueryGraphNode_Root : public UEnvironmentQueryGraphNode
 	GENERATED_UCLASS_BODY()
 
 	virtual void AllocateDefaultPins() override;
-	virtual bool CanUserDeleteNode() const override { return false; }
 	virtual FText GetNodeTitle(ENodeTitleType::Type TitleType) const override;
 };
diff --git a/Engine/Source/Editor/EnvironmentQueryEditor/Classes/EnvironmentQueryGraphNode_Test.h b/Engine/Source/Editor/EnvironmentQueryEditor/Classes/EnvironmentQueryGraphNode_Test.h
--- a/Engine/Source/Editor/EnvironmentQueryEditor/Classes/EnvironmentQueryGraphNode_Test.h
+++ b/Engine/Source/Editor/EnvironmentQueryEditor/Classes/EnvironmentQueryGraphNode_Test.h
@@ -20,12 +20,7 @@ class UEnvironmentQueryGraphNode_Test : public UEnvironmentQueryGraphNode
 	UPROPERTY()
 	uint32 bTestEnabled : 1;
 
-	UPROPERTY()
-	class UEnvironmentQueryGraphNode_Option* ParentNode;
-
-	virtual void PostPlacedNewNode() override;
-	virtual void DestroyNode() override;
-	virtual bool IsSubNode() const override { return true; }
+	virtual void InitializeInstance() override;
 	virtual FText GetNodeTitle(ENodeTitleType::Type TitleType) const override;
 	virtual FText GetDescription() const override;
 
diff --git a/Engine/Source/Editor/EnvironmentQueryEditor/EnvironmentQueryEditor.Build.cs b/Engine/Source/Editor/EnvironmentQueryEditor/EnvironmentQueryEditor.Build.cs
--- a/Engine/Source/Editor/EnvironmentQueryEditor/EnvironmentQueryEditor.Build.cs
+++ b/Engine/Source/Editor/EnvironmentQueryEditor/EnvironmentQueryEditor.Build.cs
@@ -9,6 +9,7 @@ public class EnvironmentQueryEditor : ModuleRules
         PrivateIncludePaths.AddRange(
             new string[] {
 				"Editor/GraphEditor/Private",
+				"Editor/AIGraph/Private",
 				"Editor/EnvironmentQueryEditor/Private",
 			}
 		);
@@ -37,6 +38,7 @@ public class EnvironmentQueryEditor : ModuleRules
                 "PropertyEditor",
 				"AnimGraph",
 				"BlueprintGraph",
+                "AIGraph",
                 "AIModule",
 			}
 		);
diff --git a/Engine/Source/Editor/EnvironmentQueryEditor/Private/EdGraphSchema_EnvironmentQuery.cpp b/Engine/Source/Editor/EnvironmentQueryEditor/Private/EdGraphSchema_EnvironmentQuery.cpp
--- a/Engine/Source/Editor/EnvironmentQueryEditor/Private/EdGraphSchema_EnvironmentQuery.cpp
+++ b/Engine/Source/Editor/EnvironmentQueryEditor/Private/EdGraphSchema_EnvironmentQuery.cpp
@@ -1,191 +1,20 @@
 // Copyright 1998-2015 Epic Games, Inc. All Rights Reserved.
 
 #include "EnvironmentQueryEditorPrivatePCH.h"
-#include "BlueprintGraphDefinitions.h"
-#include "GraphEditorActions.h"
-#include "EnvironmentQueryConnectionDrawingPolicy.h"
+#include "EdGraphSchema_EnvironmentQuery.h"
 #include "EnvironmentQuery/EnvQueryGenerator.h"
 #include "EnvironmentQuery/Generators/EnvQueryGenerator_Composite.h"
 #include "EnvironmentQuery/EnvQueryTest.h"
-#include "Kismet2/KismetEditorUtilities.h"
-#include "GenericCommands.h"
 
-#include "AssetRegistryModule.h"
-#include "HotReloadInterface.h"
-
-#define LOCTEXT_NAMESPACE "EnvironmentQuerySchema"
-#define SNAP_GRID (16) // @todo ensure this is the same as SNodePanel::GetSnapGridSize()
-
-namespace 
-{
-	// Maximum distance a drag can be off a node edge to require 'push off' from node
-	const int32 NodeDistance = 60;
-}
-
-UEdGraphNode* FEnvironmentQuerySchemaAction_NewNode::PerformAction(class UEdGraph* ParentGraph, UEdGraphPin* FromPin, const FVector2D Location, bool bSelectNewNode)
-{
-	UEdGraphNode* ResultNode = NULL;
-
-	// If there is a template, we actually use it
-	if (NodeTemplate != NULL)
-	{
-		NodeTemplate->SetFlags(RF_Transactional);
-
-		// set outer to be the graph so it doesn't go away
-		NodeTemplate->Rename(NULL, ParentGraph, REN_NonTransactional);
-		ParentGraph->AddNode(NodeTemplate, true);
-
-		NodeTemplate->CreateNewGuid();
-		NodeTemplate->PostPlacedNewNode();
-		NodeTemplate->AllocateDefaultPins();
-		NodeTemplate->AutowireNewNode(FromPin);
-
-		// For input pins, new node will generally overlap node being dragged off
-		// Work out if we want to visually push away from connected node
-		int32 XLocation = Location.X;
-		if (FromPin && FromPin->Direction == EGPD_Input)
-		{
-			UEdGraphNode* PinNode = FromPin->GetOwningNode();
-			const float XDelta = FMath::Abs(PinNode->NodePosX - Location.X);
-
-			if (XDelta < NodeDistance)
-			{
-				// Set location to edge of current node minus the max move distance
-				// to force node to push off from connect node enough to give selection handle
-				XLocation = PinNode->NodePosX - NodeDistance;
-			}
-		}
-
-		NodeTemplate->NodePosX = XLocation;
-		NodeTemplate->NodePosY = Location.Y;
-		NodeTemplate->SnapToGrid(SNAP_GRID);
-
-		ResultNode = NodeTemplate;
-	}
-
-	return ResultNode;
-}
-
-UEdGraphNode* FEnvironmentQuerySchemaAction_NewNode::PerformAction(class UEdGraph* ParentGraph, TArray<UEdGraphPin*>& FromPins, const FVector2D Location, bool bSelectNewNode) 
-{
-	// I have no idea if ResultNode is meaningful to assign here and return it.
-	// But every other override of this function is not void, so I'm making this one not void
-	// because the compiler complains otherwise.
-	UEdGraphNode* ResultNode = NULL;
-
-	if (FromPins.Num() > 0)
-	{
-		ResultNode = PerformAction(ParentGraph, FromPins[0], Location);
-
-		// Try autowiring the rest of the pins
-		for (int32 Index = 1; Index < FromPins.Num(); ++Index)
-		{
-			ResultNode->AutowireNewNode(FromPins[Index]);
-		}
-	}
-	else
-	{
-		ResultNode = PerformAction(ParentGraph, NULL, Location, bSelectNewNode);
-	}
-
-	return ResultNode;
-}
-
-void FEnvironmentQuerySchemaAction_NewNode::AddReferencedObjects( FReferenceCollector& Collector )
-{
-	FEdGraphSchemaAction::AddReferencedObjects( Collector );
-
-	// These don't get saved to disk, but we want to make sure the objects don't get GC'd while the action array is around
-	Collector.AddReferencedObject( NodeTemplate );
-}
-
-
-UEdGraphNode* FEnvironmentQuerySchemaAction_NewSubNode::PerformAction(class UEdGraph* ParentGraph, UEdGraphPin* FromPin, const FVector2D Location, bool bSelectNewNode)
-{
-	ParentNode->AddSubNode(NodeTemplate,ParentGraph);
-	return NULL;
-}
-
-UEdGraphNode* FEnvironmentQuerySchemaAction_NewSubNode::PerformAction(class UEdGraph* ParentGraph, TArray<UEdGraphPin*>& FromPins, const FVector2D Location, bool bSelectNewNode) 
-{
-	return PerformAction(ParentGraph, NULL, Location, bSelectNewNode);
-}
-
-void FEnvironmentQuerySchemaAction_NewSubNode::AddReferencedObjects( FReferenceCollector& Collector )
-{
-	FEdGraphSchemaAction::AddReferencedObjects( Collector );
-
-	// These don't get saved to disk, but we want to make sure the objects don't get GC'd while the action array is around
-	Collector.AddReferencedObject( NodeTemplate );
-}
+#define LOCTEXT_NAMESPACE "EnvironmentQueryEditor"
 
 //////////////////////////////////////////////////////////////////////////
 
-namespace FEQSClassCacheInfo
-{
-	static bool bEQSGeneratorCacheInvalid = true;
-	static bool bEQSTestCacheInvalid = true;
-	void InvalidateCache()
-	{
-		bEQSGeneratorCacheInvalid = true;
-		bEQSTestCacheInvalid = true;
-	}
-
-	void InvalidateCache(const class FAssetData& AssetData)
-	{
-		bEQSGeneratorCacheInvalid = true;
-		bEQSTestCacheInvalid = true;
-	}
-
-	void InvalidateCache(bool bWasTriggeredAutomatically)
-	{
-		bEQSGeneratorCacheInvalid = true;
-		bEQSTestCacheInvalid = true;
-	}
-
-	void SetupCacheObservers()
-	{
-		static bool bCacheFixDelegatesSetUp = false;
-		if (bCacheFixDelegatesSetUp == false)
-		{
-			bCacheFixDelegatesSetUp = true;
-
-			// Register with the Asset Registry to be informed when it is done loading up files.
-			FAssetRegistryModule& AssetRegistryModule = FModuleManager::GetModuleChecked<FAssetRegistryModule>(TEXT("AssetRegistry"));
-			AssetRegistryModule.Get().OnFilesLoaded().AddStatic(&FEQSClassCacheInfo::InvalidateCache);
-			AssetRegistryModule.Get().OnAssetAdded().AddStatic(&FEQSClassCacheInfo::InvalidateCache);
-			AssetRegistryModule.Get().OnAssetRemoved().AddStatic(&FEQSClassCacheInfo::InvalidateCache);
-
-			// Register to have Populate called when doing a Hot Reload.
-			IHotReloadInterface& HotReloadSupport = FModuleManager::LoadModuleChecked<IHotReloadInterface>("HotReload");
-			HotReloadSupport.OnHotReload().AddStatic(&FEQSClassCacheInfo::InvalidateCache);
-
-			// Register to have Populate called when a Blueprint is compiled.
-			GEditor->OnBlueprintCompiled().AddStatic(&FEQSClassCacheInfo::InvalidateCache);
-			GEditor->OnClassPackageLoadedOrUnloaded().AddStatic(&FEQSClassCacheInfo::InvalidateCache);
-		}
-	}
-}
-
 UEdGraphSchema_EnvironmentQuery::UEdGraphSchema_EnvironmentQuery(const FObjectInitializer& ObjectInitializer)
 	: Super(ObjectInitializer)
 {
 }
 
-TSharedPtr<FEnvironmentQuerySchemaAction_NewNode> AddNewNodeAction(FGraphContextMenuBuilder& ContextMenuBuilder, const FString& Category, const FText& MenuDesc, const FString& Tooltip)
-{
-	const TSharedRef<FEnvironmentQuerySchemaAction_NewNode> NewAction = MakeShareable( new FEnvironmentQuerySchemaAction_NewNode(Category, MenuDesc, Tooltip, 0) );
-	ContextMenuBuilder.AddAction( NewAction );
-	return NewAction;
-}
-
-TSharedPtr<FEnvironmentQuerySchemaAction_NewSubNode> AddNewSubNodeAction(FGraphContextMenuBuilder& ContextMenuBuilder, const FString& Category, const FText& MenuDesc, const FString& Tooltip)
-{
-	TSharedPtr<FEnvironmentQuerySchemaAction_NewSubNode> NewAction = MakeShareable(new FEnvironmentQuerySchemaAction_NewSubNode(Category, MenuDesc, Tooltip, 0));
-	ContextMenuBuilder.AddAction( NewAction );
-	return NewAction;
-}
-
 void UEdGraphSchema_EnvironmentQuery::CreateDefaultNodesForGraph(UEdGraph& Graph) const
 {
 	FGraphNodeCreator<UEnvironmentQueryGraphNode_Root> NodeCreator(Graph);
@@ -194,83 +23,6 @@ void UEdGraphSchema_EnvironmentQuery::CreateDefaultNodesForGraph(UEdGraph& Graph
 	SetNodeMetaData(MyNode, FNodeMetadata::DefaultGraphNode);
 }
 
-FString GetNodeDescriptionHelper(UObject* Ob)
-{
-	if (Ob == NULL)
-	{
-		return TEXT("unknown");
-	}
-
-	const UClass* ObClass = Ob->IsA(UClass::StaticClass()) ? (const UClass*)Ob : Ob->GetClass();
-	if (ObClass->HasAnyClassFlags(CLASS_CompiledFromBlueprint))
-	{
-		return ObClass->GetName().LeftChop(2);
-	}
-
-	FString TypeDesc = ObClass->GetName();
-	const int32 ShortNameIdx = TypeDesc.Find(TEXT("_"));
-	if (ShortNameIdx != INDEX_NONE)
-	{
-		TypeDesc = TypeDesc.Mid(ShortNameIdx + 1);
-	}
-
-	return TypeDesc;
-}
-
-static void GatherEQSGenerators(TArray<UClass*>& Classes)
-{
-	static TArray<UClass*> CachedClasses;
-	if (CachedClasses.Num() == 0 || FEQSClassCacheInfo::bEQSGeneratorCacheInvalid)
-	{
-		FEQSClassCacheInfo::SetupCacheObservers();
-		FEQSClassCacheInfo::bEQSGeneratorCacheInvalid = false;
-		CachedClasses.Reset();
-
-		for (TObjectIterator<UClass> It; It; ++It)
-		{
-			UClass* TestClass = *It;
-			if (TestClass->HasAnyClassFlags(CLASS_Abstract) ||
-				// ignore temporary BP classes
-				(!TestClass->HasAnyClassFlags(CLASS_Native) && (FKismetEditorUtilities::IsClassABlueprintSkeleton(TestClass) || TestClass->HasAnyClassFlags(CLASS_NewerVersionExists))) || 
-				!TestClass->IsChildOf(UEnvQueryGenerator::StaticClass()) ||
-				TestClass == UEnvQueryGenerator_Composite::StaticClass())
-			{
-				continue;
-			}
-
-			CachedClasses.Add(TestClass);
-		}
-	}
-
-	Classes.Append(CachedClasses);
-}
-
-static void GatherEQSTests(TArray<UClass*>& Classes)
-{
-	static TArray<UClass*> CachedClasses;
-	if (CachedClasses.Num() == 0 || FEQSClassCacheInfo::bEQSTestCacheInvalid)
-	{
-		FEQSClassCacheInfo::SetupCacheObservers();
-		FEQSClassCacheInfo::bEQSTestCacheInvalid = false;
-		CachedClasses.Reset();
-
-		for (TObjectIterator<UClass> It; It; ++It)
-		{
-			UClass* TestClass = *It;
-			if (TestClass->HasAnyClassFlags(CLASS_Abstract) ||
-				!TestClass->HasAnyClassFlags(CLASS_Native) ||
-				!TestClass->IsChildOf(UEnvQueryTest::StaticClass()))
-			{
-				continue;
-			}
-
-			CachedClasses.Add(TestClass);
-		}
-	}
-
-	Classes.Append(CachedClasses);
-}
-
 void UEdGraphSchema_EnvironmentQuery::GetGraphContextActions(FGraphContextMenuBuilder& ContextMenuBuilder) const
 {
 	UEnvironmentQueryGraphNode* ParentGraphNode = ContextMenuBuilder.FromPin ? Cast<UEnvironmentQueryGraphNode>(ContextMenuBuilder.FromPin->GetOuter()) : NULL;
@@ -279,123 +31,36 @@ void UEdGraphSchema_EnvironmentQuery::GetGraphContextActions(FGraphContextMenuBu
 		return;
 	}
 
-	TArray<UClass*> AllowedClasses;
-	GatherEQSGenerators(AllowedClasses);
+	FEnvironmentQueryEditorModule& EditorModule = FModuleManager::GetModuleChecked<FEnvironmentQueryEditorModule>(TEXT("EnvironmentQueryEditor"));
+	FGraphNodeClassHelper* ClassCache = EditorModule.GetClassCache().Get();
 	
-	for (int32 i = 0; i < AllowedClasses.Num(); i++)
-	{
-		TSharedPtr<FEnvironmentQuerySchemaAction_NewNode> AddOpAction = AddNewNodeAction(ContextMenuBuilder, "Generators", FText::FromString(GetNodeDescriptionHelper(AllowedClasses[i])), FString());
-		
-		UEnvironmentQueryGraphNode_Option* OpNode = NewObject<UEnvironmentQueryGraphNode_Option>(ContextMenuBuilder.OwnerOfTemporaries);
-		OpNode->EnvQueryNodeClass = AllowedClasses[i];
-		AddOpAction->NodeTemplate = OpNode;		
-	}
-}
-
-void UEdGraphSchema_EnvironmentQuery::GetGraphNodeContextActions(FGraphContextMenuBuilder& ContextMenuBuilder) const
-{
-	UEnvironmentQueryGraphNode_Option* SelectedOption = Cast<UEnvironmentQueryGraphNode_Option>(ContextMenuBuilder.SelectedObjects[0]);
-	UEdGraph* Graph = (UEdGraph*)ContextMenuBuilder.CurrentGraph;
+	TArray<FGraphNodeClassData> NodeClasses;
+	ClassCache->GatherClasses(UEnvQueryGenerator::StaticClass(), NodeClasses);
 
-	TArray<UClass*> AllowedClasses;
-	GatherEQSTests(AllowedClasses);
-
-	for (int32 i = 0; i < AllowedClasses.Num(); i++)
+	FCategorizedGraphActionListBuilder GeneratorsBuilder(TEXT("Generators"));
+	for (const auto& NodeClass : NodeClasses)
 	{
-		UEnvironmentQueryGraphNode_Test* OpNode = NewObject<UEnvironmentQueryGraphNode_Test>(Graph);
-		OpNode->EnvQueryNodeClass = AllowedClasses[i];
-
-		TSharedPtr<FEnvironmentQuerySchemaAction_NewSubNode> AddOpAction = AddNewSubNodeAction(ContextMenuBuilder, TEXT(""), FText::FromString(GetNodeDescriptionHelper(AllowedClasses[i])), "");
-		AddOpAction->ParentNode = SelectedOption;
-		AddOpAction->NodeTemplate = OpNode;
-	}
-}
+		const FText NodeTypeName = FText::FromString(FName::NameToDisplayString(NodeClass.ToString(), false));
 
-void UEdGraphSchema_EnvironmentQuery::GetContextMenuActions(const UEdGraph* CurrentGraph, const UEdGraphNode* InGraphNode, const UEdGraphPin* InGraphPin, class FMenuBuilder* MenuBuilder, bool bIsDebugging) const
-{
-	if (InGraphPin)
-	{
-		MenuBuilder->BeginSection("EnvironmentQueryGraphSchemaPinActions", LOCTEXT("PinActionsMenuHeader", "Pin Actions"));
-		{
-			// Only display the 'Break Links' option if there is a link to break!
-			if (InGraphPin->LinkedTo.Num() > 0)
-			{
-				MenuBuilder->AddMenuEntry( FGraphEditorCommands::Get().BreakPinLinks );
+		UEnvironmentQueryGraphNode_Option* OpNode = NewObject<UEnvironmentQueryGraphNode_Option>(ContextMenuBuilder.OwnerOfTemporaries);
+		OpNode->ClassData = NodeClass;
 
-				// add sub menu for break link to
-				if(InGraphPin->LinkedTo.Num() > 1)
-				{
-					MenuBuilder->AddSubMenu(
-						LOCTEXT("BreakLinkTo", "Break Link To..." ),
-						LOCTEXT("BreakSpecificLinks", "Break a specific link..." ),
-						FNewMenuDelegate::CreateUObject( (UEdGraphSchema_EnvironmentQuery*const)this, &UEdGraphSchema_EnvironmentQuery::GetBreakLinkToSubMenuActions, const_cast<UEdGraphPin*>(InGraphPin)));
-				}
-				else
-				{
-					((UEdGraphSchema_EnvironmentQuery*const)this)->GetBreakLinkToSubMenuActions(*MenuBuilder, const_cast<UEdGraphPin*>(InGraphPin));
-				}
-			}
-		}
-		MenuBuilder->EndSection();
-	}
-	else if (InGraphNode)
-	{
-		MenuBuilder->BeginSection("EnvironmentQueryGraphSchemaNodeActions", LOCTEXT("ClassActionsMenuHeader", "Actions"));
-		{
-			MenuBuilder->AddMenuEntry(FGraphEditorCommands::Get().BreakNodeLinks);
-			MenuBuilder->AddMenuEntry(FGenericCommands::Get().Delete);
-		}
-		MenuBuilder->EndSection();
+		TSharedPtr<FAISchemaAction_NewNode> AddOpAction = AddNewNodeAction(GeneratorsBuilder, NodeClass.GetCategory(), NodeTypeName, "");
+		AddOpAction->NodeTemplate = OpNode;
 	}
 
-	Super::GetContextMenuActions(CurrentGraph, InGraphNode, InGraphPin, MenuBuilder, bIsDebugging);
+	ContextMenuBuilder.Append(GeneratorsBuilder);
 }
 
-void UEdGraphSchema_EnvironmentQuery::GetBreakLinkToSubMenuActions( class FMenuBuilder& MenuBuilder, UEdGraphPin* InGraphPin )
+void UEdGraphSchema_EnvironmentQuery::GetSubNodeClasses(int32 SubNodeFlags, TArray<FGraphNodeClassData>& ClassData, UClass*& GraphNodeClass) const
 {
-	// Make sure we have a unique name for every entry in the list
-	TMap< FString, uint32 > LinkTitleCount;
-
-	// Add all the links we could break from
-	for(TArray<class UEdGraphPin*>::TConstIterator Links(InGraphPin->LinkedTo); Links; ++Links)
-	{
-		UEdGraphPin* Pin = *Links;
-		FString TitleString = Pin->GetOwningNode()->GetNodeTitle(ENodeTitleType::ListView).ToString();
-		FText Title = FText::FromString( TitleString );
-		if ( Pin->PinName != TEXT("") )
-		{
-			TitleString = FString::Printf(TEXT("%s (%s)"), *TitleString, *Pin->PinName);
-
-			// Add name of connection if possible
-			FFormatNamedArguments Args;
-			Args.Add( TEXT("NodeTitle"), Title );
-			Args.Add( TEXT("PinName"), Pin->GetDisplayName() );
-			Title = FText::Format( LOCTEXT("BreakDescPin", "{NodeTitle} ({PinName})"), Args );
-		}
-
-		uint32 &Count = LinkTitleCount.FindOrAdd( TitleString );
-
-		FText Description;
-		FFormatNamedArguments Args;
-		Args.Add( TEXT("NodeTitle"), Title );
-		Args.Add( TEXT("NumberOfNodes"), Count );
-
-		if ( Count == 0 )
-		{
-			Description = FText::Format( LOCTEXT("BreakDesc", "Break link to {NodeTitle}"), Args );
-		}
-		else
-		{
-			Description = FText::Format( LOCTEXT("BreakDescMulti", "Break link to {NodeTitle} ({NumberOfNodes})"), Args );
-		}
-		++Count;
+	FEnvironmentQueryEditorModule& EditorModule = FModuleManager::GetModuleChecked<FEnvironmentQueryEditorModule>(TEXT("EnvironmentQueryEditor"));
+	FGraphNodeClassHelper* ClassCache = EditorModule.GetClassCache().Get();
 
-		MenuBuilder.AddMenuEntry( Description, Description, FSlateIcon(), FUIAction(
-			FExecuteAction::CreateUObject((USoundClassGraphSchema*const)this, &USoundClassGraphSchema::BreakSinglePinLink, const_cast< UEdGraphPin* >(InGraphPin), *Links) ) );
-	}
+	ClassCache->GatherClasses(UEnvQueryTest::StaticClass(), ClassData);
+	GraphNodeClass = UEnvironmentQueryGraphNode_Test::StaticClass();
 }
 
-
 const FPinConnectionResponse UEdGraphSchema_EnvironmentQuery::CanCreateConnection(const UEdGraphPin* PinA, const UEdGraphPin* PinB) const
 {
 	// Make sure the pins are not on the same node
@@ -458,26 +123,4 @@ const FPinConnectionResponse UEdGraphSchema_EnvironmentQuery::CanMergeNodes(cons
 	return FPinConnectionResponse(CONNECT_RESPONSE_DISALLOW, TEXT(""));
 }
 
-FLinearColor UEdGraphSchema_EnvironmentQuery::GetPinTypeColor(const FEdGraphPinType& PinType) const
-{
-	return FColor::White;
-}
-
-bool UEdGraphSchema_EnvironmentQuery::ShouldHidePinDefaultValue(UEdGraphPin* Pin) const
-{
-	check(Pin != NULL);
-
-	if (Pin->bDefaultValueIsIgnored)
-	{
-		return true;
-	}
-
-	return false;
-}
-
-class FConnectionDrawingPolicy* UEdGraphSchema_EnvironmentQuery::CreateConnectionDrawingPolicy(int32 InBackLayerID, int32 InFrontLayerID, float InZoomFactor, const FSlateRect& InClippingRect, class FSlateWindowElementList& InDrawElements, class UEdGraph* InGraphObj) const
-{
-	return new FEnvironmentQueryConnectionDrawingPolicy(InBackLayerID, InFrontLayerID, InZoomFactor, InClippingRect, InDrawElements, InGraphObj);
-}
-
 #undef LOCTEXT_NAMESPACE
diff --git a/Engine/Source/Editor/EnvironmentQueryEditor/Private/EnvironmentQueryEditor.cpp b/Engine/Source/Editor/EnvironmentQueryEditor/Private/EnvironmentQueryEditor.cpp
--- a/Engine/Source/Editor/EnvironmentQueryEditor/Private/EnvironmentQueryEditor.cpp
+++ b/Engine/Source/Editor/EnvironmentQueryEditor/Private/EnvironmentQueryEditor.cpp
@@ -47,24 +47,6 @@ void FEnvironmentQueryEditor::UnregisterTabSpawners(const TSharedRef<class FTabM
 	TabManager->UnregisterTabSpawner( EQSUpdateGraphTabId );
 }
 
-FEnvironmentQueryEditor::FEnvironmentQueryEditor() : IEnvironmentQueryEditor()
-{
-	UEditorEngine* Editor = (UEditorEngine*)GEngine;
-	if (Editor != NULL)
-	{
-		Editor->RegisterForUndo(this);
-	}
-}
-
-FEnvironmentQueryEditor::~FEnvironmentQueryEditor()
-{
-	UEditorEngine* Editor = (UEditorEngine*)GEngine;
-	if (Editor)
-	{
-		Editor->UnregisterForUndo( this );
-	}
-}
-
 void FEnvironmentQueryEditor::InitEnvironmentQueryEditor( const EToolkitMode::Type Mode, const TSharedPtr< class IToolkitHost >& InitToolkitHost, UEnvQuery* InScript )
 {
 	Query = InScript;
@@ -145,46 +127,14 @@ TSharedRef<SGraphEditor> FEnvironmentQueryEditor::CreateGraphEditorWidget(UEdGra
 {
 	check(InGraph != NULL);
 	
-	GraphEditorCommands = MakeShareable( new FUICommandList );
-	{
-		// Editing commands
-		GraphEditorCommands->MapAction( FGenericCommands::Get().SelectAll,
-			FExecuteAction::CreateSP( this, &FEnvironmentQueryEditor::SelectAllNodes ),
-			FCanExecuteAction::CreateSP( this, &FEnvironmentQueryEditor::CanSelectAllNodes )
-			);
-
-		GraphEditorCommands->MapAction( FGenericCommands::Get().Delete,
-			FExecuteAction::CreateSP( this, &FEnvironmentQueryEditor::DeleteSelectedNodes ),
-			FCanExecuteAction::CreateSP( this, &FEnvironmentQueryEditor::CanDeleteNodes )
-			);
-
-		GraphEditorCommands->MapAction( FGenericCommands::Get().Copy,
-			FExecuteAction::CreateSP( this, &FEnvironmentQueryEditor::CopySelectedNodes ),
-			FCanExecuteAction::CreateSP( this, &FEnvironmentQueryEditor::CanCopyNodes )
-			);
-
-		GraphEditorCommands->MapAction( FGenericCommands::Get().Cut,
-			FExecuteAction::CreateSP( this, &FEnvironmentQueryEditor::CutSelectedNodes ),
-			FCanExecuteAction::CreateSP( this, &FEnvironmentQueryEditor::CanCutNodes )
-			);
-
-		GraphEditorCommands->MapAction( FGenericCommands::Get().Paste,
-			FExecuteAction::CreateSP( this, &FEnvironmentQueryEditor::PasteNodes ),
-			FCanExecuteAction::CreateSP( this, &FEnvironmentQueryEditor::CanPasteNodes )
-			);
-
-		GraphEditorCommands->MapAction( FGenericCommands::Get().Duplicate,
-			FExecuteAction::CreateSP( this, &FEnvironmentQueryEditor::DuplicateNodes ),
-			FCanExecuteAction::CreateSP( this, &FEnvironmentQueryEditor::CanDuplicateNodes )
-			);
-	}
-
 	// Create the appearance info
 	FGraphAppearanceInfo AppearanceInfo;
 	AppearanceInfo.CornerText = NSLOCTEXT("EnvironmentQueryEditor", "AppearanceCornerText", "ENVIRONMENT QUERY");
 
 	SGraphEditor::FGraphEditorEvents InEvents;
 	InEvents.OnSelectionChanged = SGraphEditor::FOnSelectionChanged::CreateSP(this, &FEnvironmentQueryEditor::OnSelectedNodesChanged);
+	
+	CreateCommandList();
 
 	// Make title bar
 	TSharedRef<SWidget> TitleBarWidget = 
@@ -334,9 +284,11 @@ void FEnvironmentQueryEditor::OnFinishedChangingProperties(const FPropertyChange
 			for (FGraphPanelSelectionSet::TConstIterator It(CurrentSelection); It; ++It)
 			{
 				UEnvironmentQueryGraphNode_Test* TestNode = Cast<UEnvironmentQueryGraphNode_Test>(*It);
-				if (TestNode && TestNode->ParentNode)
+				UEnvironmentQueryGraphNode_Option* ParentNode = TestNode ? Cast<UEnvironmentQueryGraphNode_Option>(TestNode->ParentNode) : nullptr;
+
+				if (ParentNode)
 				{
-					TestNode->ParentNode->CalculateWeights();
+					ParentNode->CalculateWeights();
 					break;
 				}
 			}
@@ -344,390 +296,4 @@ void FEnvironmentQueryEditor::OnFinishedChangingProperties(const FPropertyChange
 	}
 }
 
-void FEnvironmentQueryEditor::PostUndo(bool bSuccess)
-{	
-	if (bSuccess)
-	{
-		// Clear selection, to avoid holding refs to nodes that go away
-		TSharedPtr<SGraphEditor> CurrentGraphEditor = UpdateGraphEdPtr.Pin();
-		if (CurrentGraphEditor.IsValid())
-		{
-			CurrentGraphEditor->ClearSelectionSet();
-			CurrentGraphEditor->NotifyGraphChanged();
-		}
-		FSlateApplication::Get().DismissAllMenus();
-	}
-}
-
-void FEnvironmentQueryEditor::PostRedo(bool bSuccess)
-{
-	if (bSuccess)
-	{
-		// Clear selection, to avoid holding refs to nodes that go away
-		TSharedPtr<SGraphEditor> CurrentGraphEditor = UpdateGraphEdPtr.Pin();
-		if (CurrentGraphEditor.IsValid())
-		{
-			CurrentGraphEditor->ClearSelectionSet();
-			CurrentGraphEditor->NotifyGraphChanged();
-		}
-		FSlateApplication::Get().DismissAllMenus();
-	}
-}
-
-FGraphPanelSelectionSet FEnvironmentQueryEditor::GetSelectedNodes() const
-{
-	FGraphPanelSelectionSet CurrentSelection;
-	TSharedPtr<SGraphEditor> FocusedGraphEd = UpdateGraphEdPtr.Pin();
-	if (FocusedGraphEd.IsValid())
-	{
-		CurrentSelection = FocusedGraphEd->GetSelectedNodes();
-	}
-	return CurrentSelection;
-}
-
-void FEnvironmentQueryEditor::SelectAllNodes()
-{
-	TSharedPtr<SGraphEditor> CurrentGraphEditor = UpdateGraphEdPtr.Pin();
-	if (CurrentGraphEditor.IsValid())
-	{
-		CurrentGraphEditor->SelectAllNodes();
-	}
-}
-
-bool FEnvironmentQueryEditor::CanSelectAllNodes() const
-{
-	return true;
-}
-
-void FEnvironmentQueryEditor::DeleteSelectedNodes()
-{
-	TSharedPtr<SGraphEditor> CurrentGraphEditor = UpdateGraphEdPtr.Pin();
-	if (!CurrentGraphEditor.IsValid())
-	{
-		return;
-	}
-
-	const FScopedTransaction Transaction( FGenericCommands::Get().Delete->GetDescription() );
-	CurrentGraphEditor->GetCurrentGraph()->Modify();
-
-	const FGraphPanelSelectionSet SelectedNodes = CurrentGraphEditor->GetSelectedNodes();
-	CurrentGraphEditor->ClearSelectionSet();
-
-	for (FGraphPanelSelectionSet::TConstIterator NodeIt( SelectedNodes ); NodeIt; ++NodeIt)
-	{
-		if (UEdGraphNode* Node = Cast<UEdGraphNode>(*NodeIt))
-		{
-			if (Node->CanUserDeleteNode())
-			{
-				Node->Modify();
-				Node->DestroyNode();
-			}
-		}
-	}
-}
-
-bool FEnvironmentQueryEditor::CanDeleteNodes() const
-{
-	// If any of the nodes can be deleted then we should allow deleting
-	const FGraphPanelSelectionSet SelectedNodes = GetSelectedNodes();
-	for (FGraphPanelSelectionSet::TConstIterator SelectedIter(SelectedNodes); SelectedIter; ++SelectedIter)
-	{
-		UEdGraphNode* Node = Cast<UEdGraphNode>(*SelectedIter);
-		if (Node && Node->CanUserDeleteNode())
-		{
-			return true;
-		}
-	}
-
-	return false;
-}
-
-void FEnvironmentQueryEditor::DeleteSelectedDuplicatableNodes()
-{
-	TSharedPtr<SGraphEditor> CurrentGraphEditor = UpdateGraphEdPtr.Pin();
-	if (!CurrentGraphEditor.IsValid())
-	{
-		return;
-	}
-
-	const FGraphPanelSelectionSet OldSelectedNodes = CurrentGraphEditor->GetSelectedNodes();
-	CurrentGraphEditor->ClearSelectionSet();
-
-	for (FGraphPanelSelectionSet::TConstIterator SelectedIter(OldSelectedNodes); SelectedIter; ++SelectedIter)
-	{
-		UEdGraphNode* Node = Cast<UEdGraphNode>(*SelectedIter);
-		if (Node && Node->CanDuplicateNode())
-		{
-			CurrentGraphEditor->SetNodeSelection(Node, true);
-		}
-	}
-
-	// Delete the duplicatable nodes
-	DeleteSelectedNodes();
-
-	CurrentGraphEditor->ClearSelectionSet();
-
-	for (FGraphPanelSelectionSet::TConstIterator SelectedIter(OldSelectedNodes); SelectedIter; ++SelectedIter)
-	{
-		if (UEdGraphNode* Node = Cast<UEdGraphNode>(*SelectedIter))
-		{
-			CurrentGraphEditor->SetNodeSelection(Node, true);
-		}
-	}
-}
-
-void FEnvironmentQueryEditor::CutSelectedNodes()
-{
-	CopySelectedNodes();
-	DeleteSelectedDuplicatableNodes();
-}
-
-bool FEnvironmentQueryEditor::CanCutNodes() const
-{
-	return CanCopyNodes() && CanDeleteNodes();
-}
-
-void FEnvironmentQueryEditor::CopySelectedNodes()
-{
-	// Export the selected nodes and place the text on the clipboard
-	FGraphPanelSelectionSet SelectedNodes = GetSelectedNodes();
-	TArray<UEnvironmentQueryGraphNode*> SubNodes;
-
-	FString ExportedText;
-
-	int32 CopySubNodeIndex = 0;
-	for (FGraphPanelSelectionSet::TIterator SelectedIter(SelectedNodes); SelectedIter; ++SelectedIter)
-	{
-		UEnvironmentQueryGraphNode_Option* OptionNode = Cast<UEnvironmentQueryGraphNode_Option>(*SelectedIter);
-		UEnvironmentQueryGraphNode_Test* TestNode = Cast<UEnvironmentQueryGraphNode_Test>(*SelectedIter);
-		UEnvironmentQueryGraphNode* Node = OptionNode ? (UEnvironmentQueryGraphNode*)OptionNode : (UEnvironmentQueryGraphNode*)TestNode;
-
-		if (OptionNode == NULL && TestNode == NULL)
-		{
-			SelectedIter.RemoveCurrent();
-			continue;
-		}
-
-		Node->PrepareForCopying();
-		Node->CopySubNodeIndex = CopySubNodeIndex;
-
-		// append all subnodes for selection
-		if (OptionNode)
-		{
-			for (int32 i = 0; i < OptionNode->Tests.Num(); i++)
-			{
-				OptionNode->Tests[i]->CopySubNodeIndex = CopySubNodeIndex;
-				SubNodes.Add(OptionNode->Tests[i]);
-			}
-		}
-
-		CopySubNodeIndex++;
-	}
-
-	for (int32 i = 0; i < SubNodes.Num(); i++)
-	{
-		SelectedNodes.Add(SubNodes[i]);
-		SubNodes[i]->PrepareForCopying();
-	}
-
-	FEdGraphUtilities::ExportNodesToText(SelectedNodes, ExportedText);
-	FPlatformMisc::ClipboardCopy(*ExportedText);
-
-	for (FGraphPanelSelectionSet::TIterator SelectedIter(SelectedNodes); SelectedIter; ++SelectedIter)
-	{
-		UEnvironmentQueryGraphNode* Node = Cast<UEnvironmentQueryGraphNode>(*SelectedIter);
-		Node->PostCopyNode();
-	}
-}
-
-bool FEnvironmentQueryEditor::CanCopyNodes() const
-{
-	// If any of the nodes can be duplicated then we should allow copying
-	const FGraphPanelSelectionSet SelectedNodes = GetSelectedNodes();
-	for (FGraphPanelSelectionSet::TConstIterator SelectedIter(SelectedNodes); SelectedIter; ++SelectedIter)
-	{
-		UEdGraphNode* Node = Cast<UEdGraphNode>(*SelectedIter);
-		if (Node && Node->CanDuplicateNode())
-		{
-			return true;
-		}
-	}
-
-	return false;
-}
-
-void FEnvironmentQueryEditor::PasteNodes()
-{
-	TSharedPtr<SGraphEditor> CurrentGraphEditor = UpdateGraphEdPtr.Pin();
-	if (CurrentGraphEditor.IsValid())
-	{
-		PasteNodesHere(CurrentGraphEditor->GetPasteLocation());
-	}
-}
-
-void FEnvironmentQueryEditor::PasteNodesHere(const FVector2D& Location)
-{
-	TSharedPtr<SGraphEditor> CurrentGraphEditor = UpdateGraphEdPtr.Pin();
-	if (!CurrentGraphEditor.IsValid())
-	{
-		return;
-	}
-
-	// Undo/Redo support
-	const FScopedTransaction Transaction( FGenericCommands::Get().Paste->GetDescription() );
-	UEnvironmentQueryGraph* EdGraph = Cast<UEnvironmentQueryGraph>(CurrentGraphEditor->GetCurrentGraph());
-	EdGraph->Modify();
-	EdGraph->LockUpdates();
-
-	// Find currently selected option to paste all new tests into
-	UEnvironmentQueryGraphNode_Option* SelectedOption = NULL;
-	bool bHasMultipleOptionsSelected = false;
-
-	const FGraphPanelSelectionSet SelectedNodes = GetSelectedNodes();
-	for (FGraphPanelSelectionSet::TConstIterator SelectedIter(SelectedNodes); SelectedIter; ++SelectedIter)
-	{
-		UEnvironmentQueryGraphNode_Option* Node = Cast<UEnvironmentQueryGraphNode_Option>(*SelectedIter);
-		UEnvironmentQueryGraphNode_Test* TestNode = Cast<UEnvironmentQueryGraphNode_Test>(*SelectedIter);
-
-		if (TestNode)
-		{
-			Node = TestNode->ParentNode;
-		}
-
-		if (Node)
-		{
-			if (SelectedOption == NULL)
-			{
-				SelectedOption = Node;
-			}
-			else
-			{
-				bHasMultipleOptionsSelected = true;
-				break;
-			}
-		}
-	}
-
-	// Clear the selection set (newly pasted stuff will be selected)
-	CurrentGraphEditor->ClearSelectionSet();
-
-	// Grab the text to paste from the clipboard.
-	FString TextToImport;
-	FPlatformMisc::ClipboardPaste(TextToImport);
-
-	// Import the nodes
-	TSet<UEdGraphNode*> PastedNodes;
-	FEdGraphUtilities::ImportNodesFromText(EdGraph, TextToImport, /*out*/ PastedNodes);
-
-	//Average position of nodes so we can move them while still maintaining relative distances to each other
-	FVector2D AvgNodePosition(0.0f,0.0f);
-
-	for (TSet<UEdGraphNode*>::TIterator It(PastedNodes); It; ++It)
-	{
-		UEnvironmentQueryGraphNode* BTNode = Cast<UEnvironmentQueryGraphNode>(*It);
-		if (BTNode && !BTNode->IsSubNode())
-		{
-			AvgNodePosition.X += BTNode->NodePosX;
-			AvgNodePosition.Y += BTNode->NodePosY;
-		}
-	}
-
-	if (PastedNodes.Num() > 0)
-	{
-		float InvNumNodes = 1.0f/float(PastedNodes.Num());
-		AvgNodePosition.X *= InvNumNodes;
-		AvgNodePosition.Y *= InvNumNodes;
-	}
-
-	bool bPastedOptionNode = false;
-
-	TMap<int32, UEnvironmentQueryGraphNode_Option*> ParentMap;
-	for (TSet<UEdGraphNode*>::TIterator It(PastedNodes); It; ++It)
-	{
-		UEnvironmentQueryGraphNode* PasteNode = Cast<UEnvironmentQueryGraphNode>(*It);
-		if (PasteNode && !PasteNode->IsSubNode())
-		{
-			bPastedOptionNode = true;
-
-			// Select the newly pasted stuff
-			CurrentGraphEditor->SetNodeSelection(PasteNode, true);
-
-			PasteNode->NodePosX = (PasteNode->NodePosX - AvgNodePosition.X) + Location.X ;
-			PasteNode->NodePosY = (PasteNode->NodePosY - AvgNodePosition.Y) + Location.Y ;
-
-			PasteNode->SnapToGrid(16);
-
-			// Give new node a different Guid from the old one
-			PasteNode->CreateNewGuid();
-
-			UEnvironmentQueryGraphNode_Option* PasteOption = Cast<UEnvironmentQueryGraphNode_Option>(PasteNode);
-			if (PasteOption)
-			{
-				PasteOption->Tests.Reset();
-				ParentMap.Add(PasteOption->CopySubNodeIndex, PasteOption);
-			}			
-		}
-	}
-
-	for (TSet<UEdGraphNode*>::TIterator It(PastedNodes); It; ++It)
-	{
-		UEnvironmentQueryGraphNode_Test* PasteTest = Cast<UEnvironmentQueryGraphNode_Test>(*It);
-		if (PasteTest)
-		{
-			PasteTest->NodePosX = 0;
-			PasteTest->NodePosY = 0;
-
-			// Give new node a different Guid from the old one
-			PasteTest->CreateNewGuid();
-
-			// remove subnode from graph, it will be referenced from parent node
-			PasteTest->DestroyNode();
-
-			PasteTest->ParentNode = ParentMap.FindRef(PasteTest->CopySubNodeIndex);
-			if (PasteTest->ParentNode)
-			{
-				PasteTest->ParentNode->Tests.Add(PasteTest);
-			}
-			else if (!bHasMultipleOptionsSelected && !bPastedOptionNode && SelectedOption)
-			{
-				PasteTest->ParentNode = SelectedOption;
-				SelectedOption->Tests.Add(PasteTest);
-				SelectedOption->Modify();
-			}
-		}
-	}
-
-	// Update UI
-	EdGraph->UnlockUpdates();
-	CurrentGraphEditor->NotifyGraphChanged();
-
-	Query->PostEditChange();
-	Query->MarkPackageDirty();
-}
-
-bool FEnvironmentQueryEditor::CanPasteNodes() const
-{
-	TSharedPtr<SGraphEditor> CurrentGraphEditor = UpdateGraphEdPtr.Pin();
-	if (!CurrentGraphEditor.IsValid())
-	{
-		return false;
-	}
-
-	FString ClipboardContent;
-	FPlatformMisc::ClipboardPaste(ClipboardContent);
-
-	return FEdGraphUtilities::CanImportNodesFromText(CurrentGraphEditor->GetCurrentGraph(), ClipboardContent);
-}
-
-void FEnvironmentQueryEditor::DuplicateNodes()
-{
-	CopySelectedNodes();
-	PasteNodes();
-}
-
-bool FEnvironmentQueryEditor::CanDuplicateNodes() const
-{
-	return CanCopyNodes();
-}
-
 #undef LOCTEXT_NAMESPACE
\ No newline at end of file
diff --git a/Engine/Source/Editor/EnvironmentQueryEditor/Private/EnvironmentQueryEditor.h b/Engine/Source/Editor/EnvironmentQueryEditor/Private/EnvironmentQueryEditor.h
--- a/Engine/Source/Editor/EnvironmentQueryEditor/Private/EnvironmentQueryEditor.h
+++ b/Engine/Source/Editor/EnvironmentQueryEditor/Private/EnvironmentQueryEditor.h
@@ -4,11 +4,9 @@
 
 #include "IEnvironmentQueryEditor.h"
 #include "Toolkits/AssetEditorToolkit.h"
-#include "EditorUndoClient.h"
+#include "AIGraphEditor.h"
 
-
-/** Viewer/editor for a DataTable */
-class FEnvironmentQueryEditor : public IEnvironmentQueryEditor, public FEditorUndoClient
+class FEnvironmentQueryEditor : public IEnvironmentQueryEditor, public FAIGraphEditor
 {
 
 public:
@@ -17,10 +15,6 @@ public:
 	virtual void UnregisterTabSpawners(const TSharedRef<class FTabManager>& TabManager) override;
 
 	void InitEnvironmentQueryEditor( const EToolkitMode::Type Mode, const TSharedPtr< class IToolkitHost >& InitToolkitHost, UEnvQuery* Script );
-	FGraphPanelSelectionSet GetSelectedNodes() const;
-
-	FEnvironmentQueryEditor();
-	virtual ~FEnvironmentQueryEditor();
 
 	// Begin IToolkit interface
 	virtual FName GetToolkitFName() const override;
@@ -29,27 +23,6 @@ public:
 	virtual FLinearColor GetWorldCentricTabColorScale() const override;
 	// End IToolkit interface
 
-	// Begin FEditorUndoClient Interface
-	virtual void	PostUndo(bool bSuccess) override;
-	virtual void	PostRedo(bool bSuccess) override;
-	// End of FEditorUndoClient
-
-	// Delegates
-	void SelectAllNodes();
-	bool CanSelectAllNodes() const;
-	void DeleteSelectedNodes();
-	bool CanDeleteNodes() const;
-	void DeleteSelectedDuplicatableNodes();
-	void CutSelectedNodes();
-	bool CanCutNodes() const;
-	void CopySelectedNodes();
-	bool CanCopyNodes() const;
-	void PasteNodes();
-	void PasteNodesHere(const FVector2D& Location);
-	bool CanPasteNodes() const;
-	void DuplicateNodes();
-	bool CanDuplicateNodes() const;
-
 	void OnFinishedChangingProperties(const FPropertyChangedEvent& PropertyChangedEvent);
 
 protected:
@@ -74,15 +47,9 @@ private:
 	/* Query being edited */
 	UEnvQuery* Query;
 
-	/** */
-	TWeakPtr<SGraphEditor> UpdateGraphEdPtr;
-
 	/** Property View */
 	TSharedPtr<class IDetailsView> DetailsView;
 
-	/** The command list for this editor */
-	TSharedPtr<FUICommandList> GraphEditorCommands;
-
 	/**	Graph editor tab */
 	static const FName EQSUpdateGraphTabId;
 	static const FName EQSPropertiesTabId;
diff --git a/Engine/Source/Editor/EnvironmentQueryEditor/Private/EnvironmentQueryEditorModule.cpp b/Engine/Source/Editor/EnvironmentQueryEditor/Private/EnvironmentQueryEditorModule.cpp
--- a/Engine/Source/Editor/EnvironmentQueryEditor/Private/EnvironmentQueryEditorModule.cpp
+++ b/Engine/Source/Editor/EnvironmentQueryEditor/Private/EnvironmentQueryEditorModule.cpp
@@ -12,6 +12,8 @@
 #include "SGraphNode_EnvironmentQuery.h"
 #include "EdGraphUtilities.h"
 
+#include "EnvironmentQuery/Generators/EnvQueryGenerator_BlueprintBase.h"
+
 
 IMPLEMENT_MODULE( FEnvironmentQueryEditorModule, EnvironmentQueryEditor );
 DEFINE_LOG_CATEGORY(LogEnvironmentQueryEditor);
@@ -63,6 +65,7 @@ void FEnvironmentQueryEditorModule::ShutdownModule()
 		return;
 	}
 
+	ClassCache.Reset();
 	MenuExtensibilityManager.Reset();
 	ToolBarExtensibilityManager.Reset();
 
@@ -98,10 +101,16 @@ void FEnvironmentQueryEditorModule::ShutdownModule()
 }
 
 
-TSharedRef<IEnvironmentQueryEditor> FEnvironmentQueryEditorModule::CreateEnvironmentQueryEditor( const EToolkitMode::Type Mode, const TSharedPtr< IToolkitHost >& InitToolkitHost, UEnvQuery* Script )
+TSharedRef<IEnvironmentQueryEditor> FEnvironmentQueryEditorModule::CreateEnvironmentQueryEditor( const EToolkitMode::Type Mode, const TSharedPtr< IToolkitHost >& InitToolkitHost, UEnvQuery* Query)
 {
+	if (!ClassCache.IsValid())
+	{
+		ClassCache = MakeShareable(new FGraphNodeClassHelper(UEnvQueryNode::StaticClass()));
+		FGraphNodeClassHelper::AddObservedBlueprintClasses(UEnvQueryGenerator_BlueprintBase::StaticClass());
+		ClassCache->UpdateAvailableBlueprintClasses();
+	}
+
 	TSharedRef< FEnvironmentQueryEditor > NewEnvironmentQueryEditor( new FEnvironmentQueryEditor() );
-	NewEnvironmentQueryEditor->InitEnvironmentQueryEditor( Mode, InitToolkitHost, Script );
+	NewEnvironmentQueryEditor->InitEnvironmentQueryEditor(Mode, InitToolkitHost, Query);
 	return NewEnvironmentQueryEditor;
 }
-
diff --git a/Engine/Source/Editor/EnvironmentQueryEditor/Private/EnvironmentQueryGraph.cpp b/Engine/Source/Editor/EnvironmentQueryEditor/Private/EnvironmentQueryGraph.cpp
--- a/Engine/Source/Editor/EnvironmentQueryEditor/Private/EnvironmentQueryGraph.cpp
+++ b/Engine/Source/Editor/EnvironmentQueryEditor/Private/EnvironmentQueryGraph.cpp
@@ -2,7 +2,6 @@
 #include "EnvironmentQueryEditorPrivatePCH.h"
 #include "EnvironmentQueryEditorModule.h"
 
-
 //////////////////////////////////////////////////////////////////////////
 // EnvironmentQueryGraph
 
@@ -11,14 +10,14 @@ namespace EQSGraphVersion
 	const int32 Initial = 0;
 	const int32 NestedNodes = 1;
 	const int32 CopyPasteOutersBug = 2;
+	const int32 BlueprintClasses = 3;
 
-	const int32 Latest = CopyPasteOutersBug;
+	const int32 Latest = BlueprintClasses;
 }
 
 UEnvironmentQueryGraph::UEnvironmentQueryGraph(const FObjectInitializer& ObjectInitializer) : Super(ObjectInitializer)
 {
 	Schema = UEdGraphSchema_EnvironmentQuery::StaticClass();
-	bLockUpdates = false;
 }
 
 struct FCompareNodeXLocation
@@ -29,18 +28,18 @@ struct FCompareNodeXLocation
 	}
 };
 
-void UEnvironmentQueryGraph::UpdateAsset()
+void UEnvironmentQueryGraph::UpdateAsset(int32 UpdateFlags)
 {
-	if (bLockUpdates)
+	if (IsLocked())
 	{
 		return;
 	}
 
-	//let's find root node
+	// let's find root node
 	UEnvironmentQueryGraphNode_Root* RootNode = NULL;
-	for (int32 Index = 0; Index < Nodes.Num(); ++Index)
+	for (int32 Idx = 0; Idx < Nodes.Num(); Idx++)
 	{
-		RootNode = Cast<UEnvironmentQueryGraphNode_Root>(Nodes[Index]);
+		RootNode = Cast<UEnvironmentQueryGraphNode_Root>(Nodes[Idx]);
 		if (RootNode != NULL)
 		{
 			break;
@@ -56,9 +55,9 @@ void UEnvironmentQueryGraph::UpdateAsset()
 		// sort connections so that they're organized the same as user can see in the editor
 		MyPin->LinkedTo.Sort(FCompareNodeXLocation());
 
-		for (int32 Index=0; Index < MyPin->LinkedTo.Num(); ++Index)
+		for (int32 Idx = 0; Idx < MyPin->LinkedTo.Num(); Idx++)
 		{
-			UEnvironmentQueryGraphNode_Option* OptionNode = Cast<UEnvironmentQueryGraphNode_Option>(MyPin->LinkedTo[Index]->GetOwningNode());
+			UEnvironmentQueryGraphNode_Option* OptionNode = Cast<UEnvironmentQueryGraphNode_Option>(MyPin->LinkedTo[Idx]->GetOwningNode());
 			if (OptionNode)
 			{
 				UEnvQueryOption* OptionInstance = Cast<UEnvQueryOption>(OptionNode->NodeInstance);
@@ -66,9 +65,9 @@ void UEnvironmentQueryGraph::UpdateAsset()
 				{
 					OptionInstance->Tests.Reset();
 
-					for (int32 iTest = 0; iTest < OptionNode->Tests.Num(); iTest++)
+					for (int32 TestIdx = 0; TestIdx < OptionNode->SubNodes.Num(); TestIdx++)
 					{
-						UEnvironmentQueryGraphNode_Test* TestNode = OptionNode->Tests[iTest];
+						UEnvironmentQueryGraphNode_Test* TestNode = Cast<UEnvironmentQueryGraphNode_Test>(OptionNode->SubNodes[TestIdx]);
 						if (TestNode && TestNode->bTestEnabled)
 						{
 							UEnvQueryTest* TestInstance = Cast<UEnvQueryTest>(TestNode->NodeInstance);
@@ -93,9 +92,9 @@ void UEnvironmentQueryGraph::UpdateAsset()
 
 void UEnvironmentQueryGraph::CalculateAllWeights()
 {
-	for (int32 i = 0; i < Nodes.Num(); i++)
+	for (int32 Idx = 0; Idx < Nodes.Num(); Idx++)
 	{
-		UEnvironmentQueryGraphNode_Option* OptionNode = Cast<UEnvironmentQueryGraphNode_Option>(Nodes[i]);
+		UEnvironmentQueryGraphNode_Option* OptionNode = Cast<UEnvironmentQueryGraphNode_Option>(Nodes[Idx]);
 		if (OptionNode)
 		{
 			OptionNode->CalculateWeights();
@@ -126,15 +125,20 @@ void UEnvironmentQueryGraph::UpdateVersion()
 		UpdateVersion_FixupOuters();
 	}
 
+	if (GraphVersion < EQSGraphVersion::BlueprintClasses)
+	{
+		UpdateVersion_CollectClassData();
+	}
+
 	GraphVersion = EQSGraphVersion::Latest;
 	Modify();
 }
 
 void UEnvironmentQueryGraph::UpdateVersion_NestedNodes()
 {
-	for (int32 i = 0; i < Nodes.Num(); i++)
+	for (int32 Idx = 0; Idx < Nodes.Num(); Idx++)
 	{
-		UEnvironmentQueryGraphNode_Option* OptionNode = Cast<UEnvironmentQueryGraphNode_Option>(Nodes[i]);
+		UEnvironmentQueryGraphNode_Option* OptionNode = Cast<UEnvironmentQueryGraphNode_Option>(Nodes[Idx]);
 		if (OptionNode)
 		{
 			UEnvironmentQueryGraphNode* Node = OptionNode;
@@ -153,7 +157,7 @@ void UEnvironmentQueryGraph::UpdateVersion_NestedNodes()
 							if (LinkedTest)
 							{
 								LinkedTest->ParentNode = OptionNode;
-								OptionNode->Tests.Add(LinkedTest);
+								OptionNode->SubNodes.Add(LinkedTest);
 
 								NextNode = LinkedTest;
 								break;
@@ -169,17 +173,17 @@ void UEnvironmentQueryGraph::UpdateVersion_NestedNodes()
 		}
 	}
 
-	for (int32 i = Nodes.Num() - 1; i >= 0; i--)
+	for (int32 Idx = Nodes.Num() - 1; Idx >= 0; Idx--)
 	{
-		UEnvironmentQueryGraphNode_Test* TestNode = Cast<UEnvironmentQueryGraphNode_Test>(Nodes[i]);
+		UEnvironmentQueryGraphNode_Test* TestNode = Cast<UEnvironmentQueryGraphNode_Test>(Nodes[Idx]);
 		if (TestNode)
 		{
 			TestNode->Pins.Empty();
-			Nodes.RemoveAt(i);
+			Nodes.RemoveAt(Idx);
 			continue;
 		}
 
-		UEnvironmentQueryGraphNode_Option* OptionNode = Cast<UEnvironmentQueryGraphNode_Option>(Nodes[i]);
+		UEnvironmentQueryGraphNode_Option* OptionNode = Cast<UEnvironmentQueryGraphNode_Option>(Nodes[Idx]);
 		if (OptionNode && OptionNode->Pins.IsValidIndex(1))
 		{
 			OptionNode->Pins.RemoveAt(1);
@@ -189,9 +193,9 @@ void UEnvironmentQueryGraph::UpdateVersion_NestedNodes()
 
 void UEnvironmentQueryGraph::UpdateVersion_FixupOuters()
 {
-	for (int32 Index = 0; Index < Nodes.Num(); ++Index)
+	for (int32 Idx = 0; Idx < Nodes.Num(); Idx++)
 	{
-		UEnvironmentQueryGraphNode* MyNode = Cast<UEnvironmentQueryGraphNode>(Nodes[Index]);
+		UEnvironmentQueryGraphNode* MyNode = Cast<UEnvironmentQueryGraphNode>(Nodes[Idx]);
 		if (MyNode)
 		{
 			MyNode->PostEditImport();
@@ -199,65 +203,55 @@ void UEnvironmentQueryGraph::UpdateVersion_FixupOuters()
 	}
 }
 
-void UEnvironmentQueryGraph::RemoveOrphanedNodes()
+void UEnvironmentQueryGraph::UpdateVersion_CollectClassData()
 {
-	UEnvQuery* QueryAsset = CastChecked<UEnvQuery>(GetOuter());
-
-	// Obtain a list of all nodes that should be in the asset
-	TSet<UObject*> AllNodes;
-	for (int32 Index = 0; Index < Nodes.Num(); ++Index)
+	for (int32 Idx = 0; Idx < Nodes.Num(); Idx++)
 	{
-		UEnvironmentQueryGraphNode_Option* OptionNode = Cast<UEnvironmentQueryGraphNode_Option>(Nodes[Index]);
-		if (OptionNode)
+		UEnvironmentQueryGraphNode* MyNode = Cast<UEnvironmentQueryGraphNode>(Nodes[Idx]);
+		if (MyNode)
 		{
-			UEnvQueryOption* OptionInstance = Cast<UEnvQueryOption>(OptionNode->NodeInstance);
-			if (OptionInstance)
+			UEnvQueryOption* OptionInstance = Cast<UEnvQueryOption>(MyNode->NodeInstance);
+			if (OptionInstance && OptionInstance->Generator)
 			{
-				AllNodes.Add(OptionInstance);
-				if (OptionInstance->Generator)
-				{
-					AllNodes.Add(OptionInstance->Generator);
-				}
+				UpdateNodeClassData(MyNode, OptionInstance->Generator->GetClass());
 			}
 
-			for (int32 SubIdx = 0; SubIdx < OptionNode->Tests.Num(); SubIdx++)
+			for (int32 SubIdx = 0; SubIdx < MyNode->SubNodes.Num(); SubIdx++)
 			{
-				if (OptionNode->Tests[SubIdx] && OptionNode->Tests[SubIdx]->NodeInstance)
+				UAIGraphNode* SubNode = MyNode->SubNodes[SubIdx];
+				if (SubNode && SubNode->NodeInstance)
 				{
-					AllNodes.Add(OptionNode->Tests[SubIdx]->NodeInstance);
+					UpdateNodeClassData(SubNode, SubNode->NodeInstance->GetClass());
 				}
 			}
 		}
 	}
-
-	// Obtain a list of all nodes actually in the asset and discard unused nodes
-	TArray<UObject*> AllInners;
-	const bool bIncludeNestedObjects = false;
-	GetObjectsWithOuter(QueryAsset, AllInners, bIncludeNestedObjects);
-	for (auto InnerIt = AllInners.CreateConstIterator(); InnerIt; ++InnerIt)
-	{
-		UObject* Node = *InnerIt;
-		const bool bEQSNode =
-			Node->IsA(UEnvQueryGenerator::StaticClass()) ||
-			Node->IsA(UEnvQueryTest::StaticClass()) ||
-			Node->IsA(UEnvQueryOption::StaticClass());
-
-		if (bEQSNode && !AllNodes.Contains(Node))
-		{
-			Node->SetFlags(RF_Transient);
-			Node->Rename(NULL, GetTransientPackage(), REN_DontCreateRedirectors | REN_NonTransactional | REN_ForceNoResetLoaders);
-		}
-	}
 }
 
-void UEnvironmentQueryGraph::LockUpdates()
+void UEnvironmentQueryGraph::UpdateNodeClassData(UAIGraphNode* UpdateNode, UClass* InstanceClass)
 {
-	bLockUpdates = true;
+	UBlueprint* BPOwner = Cast<UBlueprint>(InstanceClass->ClassGeneratedBy);
+	if (BPOwner)
+	{
+		UpdateNode->ClassData = FGraphNodeClassData(BPOwner->GetName(), BPOwner->GetOutermost()->GetName(), InstanceClass->GetName(), InstanceClass);
+	}
+	else
+	{
+		UpdateNode->ClassData = FGraphNodeClassData(InstanceClass, "");
+	}
 }
 
-void UEnvironmentQueryGraph::UnlockUpdates()
+void UEnvironmentQueryGraph::CollectAllNodeInstances(TSet<UObject*>& NodeInstances)
 {
-	bLockUpdates = false;
+	Super::CollectAllNodeInstances(NodeInstances);
 
-	UpdateAsset();
+	for (int32 Idx = 0; Idx < Nodes.Num(); Idx++)
+	{
+		UEnvironmentQueryGraphNode* MyNode = Cast<UEnvironmentQueryGraphNode>(Nodes[Idx]);
+		UEnvQueryOption* OptionInstance = MyNode ? Cast<UEnvQueryOption>(MyNode->NodeInstance) : nullptr;
+		if (OptionInstance && OptionInstance->Generator)
+		{
+			NodeInstances.Add(OptionInstance->Generator);
+		}
+	}
 }
diff --git a/Engine/Source/Editor/EnvironmentQueryEditor/Private/EnvironmentQueryGraphNode.cpp b/Engine/Source/Editor/EnvironmentQueryEditor/Private/EnvironmentQueryGraphNode.cpp
--- a/Engine/Source/Editor/EnvironmentQueryEditor/Private/EnvironmentQueryGraphNode.cpp
+++ b/Engine/Source/Editor/EnvironmentQueryEditor/Private/EnvironmentQueryGraphNode.cpp
@@ -1,43 +1,11 @@
 // Copyright 1998-2015 Epic Games, Inc. All Rights Reserved.
 
 #include "EnvironmentQueryEditorPrivatePCH.h"
-#include "ScopedTransaction.h"
 
-UEnvironmentQueryGraphNode::UEnvironmentQueryGraphNode(const FObjectInitializer& ObjectInitializer)
-	: Super(ObjectInitializer)
-	, NodeInstance(NULL)
-	, EnvQueryNodeClass(NULL)
+UEnvironmentQueryGraphNode::UEnvironmentQueryGraphNode(const FObjectInitializer& ObjectInitializer) : Super(ObjectInitializer)
 {
 }
 
-void UEnvironmentQueryGraphNode::PrepareForCopying()
-{
-	if (NodeInstance)
-	{
-		// Temporarily take ownership of the node instance, so that it is not deleted when cutting
-		NodeInstance->Rename(NULL, this, REN_DontCreateRedirectors | REN_DoNotDirty );
-	}
-}
-
-void UEnvironmentQueryGraphNode::PostEditImport()
-{
-	ResetNodeOwner();
-}
-
-void UEnvironmentQueryGraphNode::PostCopyNode()
-{
-	ResetNodeOwner();
-}
-
-void UEnvironmentQueryGraphNode::ResetNodeOwner()
-{
-	if (NodeInstance)
-	{
-		UEnvQuery* Query = Cast<UEnvQuery>(GetEnvironmentQueryGraph()->GetOuter());
-		NodeInstance->Rename(NULL, Query, REN_DontCreateRedirectors | REN_DoNotDirty);
-	}
-}
-
 FText UEnvironmentQueryGraphNode::GetNodeTitle(ENodeTitleType::Type TitleType) const
 {
 	return FText::GetEmpty();
@@ -48,74 +16,11 @@ FText UEnvironmentQueryGraphNode::GetDescription() const
 	return FText::GetEmpty();
 }
 
-UEdGraphPin* UEnvironmentQueryGraphNode::GetInputPin(int32 InputIndex) const
-{
-	check(InputIndex >= 0);
-
-	for (int32 PinIndex = 0, FoundInputs = 0; PinIndex < Pins.Num(); PinIndex++)
-	{
-		if (Pins[PinIndex]->Direction == EGPD_Input)
-		{
-			if (InputIndex == FoundInputs)
-			{
-				return Pins[PinIndex];
-			}
-			else
-			{
-				FoundInputs++;
-			}
-		}
-	}
-
-	return NULL;
-}
-
-UEdGraphPin* UEnvironmentQueryGraphNode::GetOutputPin(int32 InputIndex) const
-{
-	check(InputIndex >= 0);
-
-	for (int32 PinIndex = 0, FoundInputs = 0; PinIndex < Pins.Num(); PinIndex++)
-	{
-		if (Pins[PinIndex]->Direction == EGPD_Output)
-		{
-			if (InputIndex == FoundInputs)
-			{
-				return Pins[PinIndex];
-			}
-			else
-			{
-				FoundInputs++;
-			}
-		}
-	}
-
-	return NULL;
-}
-
-void UEnvironmentQueryGraphNode::AutowireNewNode(UEdGraphPin* FromPin)
-{
-	Super::AutowireNewNode(FromPin);
-
-	if (FromPin != NULL)
-	{
-		if (GetSchema()->TryCreateConnection(FromPin, GetInputPin()))
-		{
-			FromPin->GetOwningNode()->NodeConnectionListChanged();
-		}
-	}
-}
-
-
 UEnvironmentQueryGraph* UEnvironmentQueryGraphNode::GetEnvironmentQueryGraph()
 {
 	return CastChecked<UEnvironmentQueryGraph>(GetGraph());
 }
 
-void UEnvironmentQueryGraphNode::NodeConnectionListChanged()
-{
-	GetEnvironmentQueryGraph()->UpdateAsset();
-}
-
 bool UEnvironmentQueryGraphNode::CanCreateUnderSpecifiedSchema(const UEdGraphSchema* DesiredSchema) const
 {
 	return DesiredSchema->GetClass()->IsChildOf(UEdGraphSchema_EnvironmentQuery::StaticClass());
diff --git a/Engine/Source/Editor/EnvironmentQueryEditor/Private/EnvironmentQueryGraphNode_Option.cpp b/Engine/Source/Editor/EnvironmentQueryEditor/Private/EnvironmentQueryGraphNode_Option.cpp
--- a/Engine/Source/Editor/EnvironmentQueryEditor/Private/EnvironmentQueryGraphNode_Option.cpp
+++ b/Engine/Source/Editor/EnvironmentQueryEditor/Private/EnvironmentQueryGraphNode_Option.cpp
@@ -2,11 +2,11 @@
 
 #include "EnvironmentQueryEditorPrivatePCH.h"
 #include "ScopedTransaction.h"
-#include "SGraphEditorActionMenu_EnvironmentQuery.h"
+#include "SGraphEditorActionMenuAI.h"
 #include "EnvironmentQuery/EnvQueryTest.h"
 #include "EnvironmentQuery/EnvQueryOption.h"
 
-#define LOCTEXT_NAMESPACE "BehaviorTreeGraphNode"
+#define LOCTEXT_NAMESPACE "EnvironmentQueryEditor"
 
 UEnvironmentQueryGraphNode_Option::UEnvironmentQueryGraphNode_Option(const FObjectInitializer& ObjectInitializer)
 	: Super(ObjectInitializer)
@@ -20,17 +20,23 @@ void UEnvironmentQueryGraphNode_Option::AllocateDefaultPins()
 
 void UEnvironmentQueryGraphNode_Option::PostPlacedNewNode()
 {
-	if (EnvQueryNodeClass != NULL)
+	UClass* NodeClass = ClassData.GetClass(true);
+	if (NodeClass)
 	{
-		UEnvQuery* Query = Cast<UEnvQuery>(GetEnvironmentQueryGraph()->GetOuter());
-		UEnvQueryOption* QueryOption = NewObject<UEnvQueryOption>(Query);
-		QueryOption->Generator = NewObject<UEnvQueryGenerator>(Query, EnvQueryNodeClass);
-		QueryOption->Generator->UpdateGeneratorVersion();
-		
-		QueryOption->SetFlags(RF_Transactional);
-		QueryOption->Generator->SetFlags(RF_Transactional);
+		UEdGraph* MyGraph = GetGraph();
+		UObject* GraphOwner = MyGraph ? MyGraph->GetOuter() : nullptr;
+		if (GraphOwner)
+		{
+			UEnvQueryOption* QueryOption = NewObject<UEnvQueryOption>(GraphOwner);
+			QueryOption->Generator = NewObject<UEnvQueryGenerator>(GraphOwner, NodeClass);
+			QueryOption->Generator->UpdateNodeVersion();
 
-		NodeInstance = QueryOption;
+			QueryOption->SetFlags(RF_Transactional);
+			QueryOption->Generator->SetFlags(RF_Transactional);
+
+			NodeInstance = QueryOption;
+			InitializeInstance();
+		}
 	}
 }
 
@@ -41,8 +47,8 @@ void UEnvironmentQueryGraphNode_Option::ResetNodeOwner()
 	UEnvQueryOption* OptionInstance = Cast<UEnvQueryOption>(NodeInstance);
 	if (OptionInstance && OptionInstance->Generator)
 	{
-		UEnvQuery* Query = Cast<UEnvQuery>(GetEnvironmentQueryGraph()->GetOuter());
-		OptionInstance->Generator->Rename(NULL, Query, REN_DontCreateRedirectors | REN_DoNotDirty);
+		UObject* GraphOwner = GetGraph() ? GetGraph()->GetOuter() : nullptr;
+		OptionInstance->Generator->Rename(NULL, GraphOwner, REN_DontCreateRedirectors | REN_DoNotDirty);
 	}
 }
 
@@ -70,36 +76,6 @@ FText UEnvironmentQueryGraphNode_Option::GetDescription() const
 	return OptionInstance ? OptionInstance->GetDescriptionDetails() : FText::GetEmpty();
 }
 
-void UEnvironmentQueryGraphNode_Option::AddSubNode(UEnvironmentQueryGraphNode_Test* NodeTemplate, class UEdGraph* ParentGraph)
-{
-	const FScopedTransaction Transaction(LOCTEXT("AddNode", "Add Node"));
-	ParentGraph->Modify();
-	Modify();
-
-	NodeTemplate->SetFlags(RF_Transactional);
-
-	// set outer to be the graph so it doesn't go away
-	NodeTemplate->Rename(NULL, ParentGraph, REN_NonTransactional);
-	NodeTemplate->ParentNode = this;
-	Tests.Add(NodeTemplate);
-
-	NodeTemplate->CreateNewGuid();
-	NodeTemplate->PostPlacedNewNode();
-	NodeTemplate->AllocateDefaultPins();
-	NodeTemplate->AutowireNewNode(NULL);
-
-	NodeTemplate->NodePosX = 0;
-	NodeTemplate->NodePosY = 0;
-
-	ParentGraph->NotifyGraphChanged();
-
-	UEnvironmentQueryGraph* MyGraph = Cast<UEnvironmentQueryGraph>(ParentGraph);
-	if (MyGraph)
-	{
-		MyGraph->UpdateAsset();
-	}
-}
-
 void UEnvironmentQueryGraphNode_Option::GetContextMenuActions(const FGraphNodeContextMenuBuilder& Context) const
 {
 	Context.MenuBuilder->AddSubMenu(
@@ -111,8 +87,8 @@ void UEnvironmentQueryGraphNode_Option::GetContextMenuActions(const FGraphNodeCo
 
 void UEnvironmentQueryGraphNode_Option::CreateAddTestSubMenu(class FMenuBuilder& MenuBuilder, UEdGraph* Graph) const
 {
-	TSharedRef<SGraphEditorActionMenu_EnvironmentQuery> Menu =	
-		SNew(SGraphEditorActionMenu_EnvironmentQuery)
+	TSharedRef<SGraphEditorActionMenuAI> Menu =	
+		SNew(SGraphEditorActionMenuAI)
 		.GraphObj( Graph )
 		.GraphNode((UEnvironmentQueryGraphNode_Option*)this)
 		.AutoExpandActionMenu(true);
@@ -123,14 +99,15 @@ void UEnvironmentQueryGraphNode_Option::CreateAddTestSubMenu(class FMenuBuilder&
 void UEnvironmentQueryGraphNode_Option::CalculateWeights()
 {
 	float MaxWeight = -1.0f;
-	for (int32 i = 0; i < Tests.Num(); i++)
+	for (int32 Idx = 0; Idx < SubNodes.Num(); Idx++)
 	{
-		if (Tests[i] == NULL || !Tests[i]->bTestEnabled)
+		UEnvironmentQueryGraphNode_Test* TestNode = Cast<UEnvironmentQueryGraphNode_Test>(SubNodes[Idx]);
+		if (TestNode == nullptr || !TestNode->bTestEnabled)
 		{
 			continue;
 		}
 
-		UEnvQueryTest* TestInstance = Cast<UEnvQueryTest>(Tests[i]->NodeInstance);
+		UEnvQueryTest* TestInstance = Cast<UEnvQueryTest>(TestNode->NodeInstance);
 		if (TestInstance && !TestInstance->Weight.IsNamedParam())
 		{
 			MaxWeight = FMath::Max(MaxWeight, FMath::Abs(TestInstance->Weight.Value));
@@ -142,20 +119,21 @@ void UEnvironmentQueryGraphNode_Option::CalculateWeights()
 		MaxWeight = 1.0f;
 	}
 
-	for (int32 i = 0; i < Tests.Num(); i++)
+	for (int32 Idx = 0; Idx < SubNodes.Num(); Idx++)
 	{
-		if (Tests[i] == NULL)
+		UEnvironmentQueryGraphNode_Test* TestNode = Cast<UEnvironmentQueryGraphNode_Test>(SubNodes[Idx]);
+		if (TestNode == NULL)
 		{
 			continue;
 		}
 		
-		UEnvQueryTest* TestInstance = Cast<UEnvQueryTest>(Tests[i]->NodeInstance);
-		const bool bHasNamed = TestInstance && Tests[i]->bTestEnabled && TestInstance->Weight.IsNamedParam();
-		const float NewWeight = (TestInstance && Tests[i]->bTestEnabled) ?
+		UEnvQueryTest* TestInstance = Cast<UEnvQueryTest>(TestNode->NodeInstance);
+		const bool bHasNamed = TestInstance && TestNode->bTestEnabled && TestInstance->Weight.IsNamedParam();
+		const float NewWeight = (TestInstance && TestNode->bTestEnabled) ?
 			(TestInstance->Weight.IsNamedParam() ? 1.0f : FMath::Clamp(FMath::Abs(TestInstance->Weight.Value) / MaxWeight, 0.0f, 1.0f)) : 
 			-1.0f;
 
-		Tests[i]->SetDisplayedWeight(NewWeight, bHasNamed);
+		TestNode->SetDisplayedWeight(NewWeight, bHasNamed);
 	}
 }
 
diff --git a/Engine/Source/Editor/EnvironmentQueryEditor/Private/EnvironmentQueryGraphNode_Root.cpp b/Engine/Source/Editor/EnvironmentQueryEditor/Private/EnvironmentQueryGraphNode_Root.cpp
--- a/Engine/Source/Editor/EnvironmentQueryEditor/Private/EnvironmentQueryGraphNode_Root.cpp
+++ b/Engine/Source/Editor/EnvironmentQueryEditor/Private/EnvironmentQueryGraphNode_Root.cpp
@@ -4,6 +4,7 @@
 
 UEnvironmentQueryGraphNode_Root::UEnvironmentQueryGraphNode_Root(const FObjectInitializer& ObjectInitializer) : Super(ObjectInitializer)
 {
+	bIsReadOnly = true;
 }
 
 void UEnvironmentQueryGraphNode_Root::AllocateDefaultPins()
@@ -13,5 +14,5 @@ void UEnvironmentQueryGraphNode_Root::AllocateDefaultPins()
 
 FText UEnvironmentQueryGraphNode_Root::GetNodeTitle(ENodeTitleType::Type TitleType) const
 {
-	return NSLOCTEXT("EnvironmentQueryGraphNode_Root", "Root", "ROOT");
+	return NSLOCTEXT("EnvironmentQueryEditor", "Root", "ROOT");
 }
diff --git a/Engine/Source/Editor/EnvironmentQueryEditor/Private/EnvironmentQueryGraphNode_Test.cpp b/Engine/Source/Editor/EnvironmentQueryEditor/Private/EnvironmentQueryGraphNode_Test.cpp
--- a/Engine/Source/Editor/EnvironmentQueryEditor/Private/EnvironmentQueryGraphNode_Test.cpp
+++ b/Engine/Source/Editor/EnvironmentQueryEditor/Private/EnvironmentQueryGraphNode_Test.cpp
@@ -5,44 +5,31 @@
 
 UEnvironmentQueryGraphNode_Test::UEnvironmentQueryGraphNode_Test(const FObjectInitializer& ObjectInitializer) : Super(ObjectInitializer)
 {
+	bIsSubNode = true;
 	bTestEnabled = true;
 	bHasNamedWeight = false;
 	TestWeightPct = -1.0f;
 }
 
-void UEnvironmentQueryGraphNode_Test::PostPlacedNewNode()
+void UEnvironmentQueryGraphNode_Test::InitializeInstance()
 {
-	if (EnvQueryNodeClass != NULL)
-	{
-		UEnvQuery* Query = Cast<UEnvQuery>(GetEnvironmentQueryGraph()->GetOuter());
-		NodeInstance = NewObject<UEnvQueryTest>(Query, EnvQueryNodeClass);
-		NodeInstance->SetFlags(RF_Transactional);
-
-		UEnvQueryTest* TestInstance = Cast<UEnvQueryTest>(NodeInstance);
-		TestInstance->UpdateTestVersion();
-	}
-
-	if (ParentNode)
+	UEnvQueryTest* TestInstance = Cast<UEnvQueryTest>(NodeInstance);
+	if (TestInstance)
 	{
-		ParentNode->CalculateWeights();
+		TestInstance->UpdateNodeVersion();
 	}
-}
 
-void UEnvironmentQueryGraphNode_Test::DestroyNode()
-{
-	if (ParentNode)
+	UEnvironmentQueryGraphNode_Option* ParentOption = Cast<UEnvironmentQueryGraphNode_Option>(ParentNode);
+	if (ParentOption)
 	{
-		ParentNode->Modify();
-		ParentNode->Tests.Remove(this);
+		ParentOption->CalculateWeights();
 	}
-
-	Super::DestroyNode();
 }
 
 FText UEnvironmentQueryGraphNode_Test::GetNodeTitle(ENodeTitleType::Type TitleType) const
 {
 	UEnvQueryTest* TestInstance = Cast<UEnvQueryTest>(NodeInstance);
-	return TestInstance ? FText::FromString(TestInstance->GetDescriptionTitle()) : FText::GetEmpty();
+	return TestInstance ? TestInstance->GetDescriptionTitle() : FText::GetEmpty();
 }
 
 FText UEnvironmentQueryGraphNode_Test::GetDescription() const
diff --git a/Engine/Source/Editor/EnvironmentQueryEditor/Private/SGraphNode_EnvironmentQuery.cpp b/Engine/Source/Editor/EnvironmentQueryEditor/Private/SGraphNode_EnvironmentQuery.cpp
--- a/Engine/Source/Editor/EnvironmentQueryEditor/Private/SGraphNode_EnvironmentQuery.cpp
+++ b/Engine/Source/Editor/EnvironmentQueryEditor/Private/SGraphNode_EnvironmentQuery.cpp
@@ -16,83 +16,20 @@
 
 #define LOCTEXT_NAMESPACE "EnvironmentQueryEditor"
 
-TSharedRef<FDragNodeEQSTimed> FDragNodeEQSTimed::New(const TSharedRef<SGraphPanel>& InGraphPanel, const TSharedRef<SGraphNode>& InDraggedNode)
-{
-	TSharedRef<FDragNodeEQSTimed> Operation = MakeShareable(new FDragNodeEQSTimed);
-	Operation->StartTime = FPlatformTime::Seconds();
-	Operation->GraphPanel = InGraphPanel;
-	Operation->DraggedNodes.Add(InDraggedNode);
-	// adjust the decorator away from the current mouse location a small amount based on cursor size
-	Operation->DecoratorAdjust = FSlateApplication::Get().GetCursorSize();
-	Operation->Construct();
-	return Operation;
-}
-
-TSharedRef<FDragNodeEQSTimed> FDragNodeEQSTimed::New(const TSharedRef<SGraphPanel>& InGraphPanel, const TArray< TSharedRef<SGraphNode> >& InDraggedNodes)
-{
-	TSharedRef<FDragNodeEQSTimed> Operation = MakeShareable(new FDragNodeEQSTimed);
-	Operation->StartTime = FPlatformTime::Seconds();
-	Operation->GraphPanel = InGraphPanel;
-	Operation->DraggedNodes = InDraggedNodes;
-	Operation->DecoratorAdjust = FSlateApplication::Get().GetCursorSize();
-	Operation->Construct();
-	return Operation;
-}
-
-UEnvironmentQueryGraphNode_Test* FDragNodeEQSTimed::GetDropTargetNode() const
-{
-	return Cast<UEnvironmentQueryGraphNode_Test>(GetHoveredNode());
-}
-
-/////////////////////////////////////////////////////
-// SStateMachineOutputPin
-
-class SEnvironmentQueryPin : public SGraphPin
+class SEnvironmentQueryPin : public SGraphPinAI
 {
 public:
 	SLATE_BEGIN_ARGS(SEnvironmentQueryPin){}
 	SLATE_END_ARGS()
 
 	void Construct(const FArguments& InArgs, UEdGraphPin* InPin);
-protected:
-	// Begin SGraphPin interface
-	virtual FSlateColor GetPinColor() const override;
-	virtual TSharedRef<SWidget>	GetDefaultValueWidget() override;
-	// End SGraphPin interface
 
-	const FSlateBrush* GetPinBorder() const;
+	virtual FSlateColor GetPinColor() const override;
 };
 
 void SEnvironmentQueryPin::Construct(const FArguments& InArgs, UEdGraphPin* InPin)
 {
-	this->SetCursor( EMouseCursor::Default );
-
-	bShowLabel = true;
-	IsEditable = true;
-
-	GraphPinObj = InPin;
-	check(GraphPinObj != NULL);
-
-	const UEdGraphSchema* Schema = GraphPinObj->GetSchema();
-	check(Schema);
-
-	SBorder::Construct( SBorder::FArguments()
-		.BorderImage( this, &SEnvironmentQueryPin::GetPinBorder )
-		.BorderBackgroundColor( this, &SEnvironmentQueryPin::GetPinColor )
-		.OnMouseButtonDown( this, &SEnvironmentQueryPin::OnPinMouseDown )
-		.Cursor( this, &SEnvironmentQueryPin::GetPinCursor )
-		.Padding(FMargin(5.0f))
-		);
-}
-
-TSharedRef<SWidget>	SEnvironmentQueryPin::GetDefaultValueWidget()
-{
-	return SNew(STextBlock);
-}
-
-const FSlateBrush* SEnvironmentQueryPin::GetPinBorder() const
-{
-	return FEditorStyle::GetBrush(TEXT("Graph.StateNode.Body"));
+	SGraphPinAI::Construct(SGraphPinAI::FArguments(), InPin);
 }
 
 FSlateColor SEnvironmentQueryPin::GetPinColor() const
@@ -106,22 +43,17 @@ FSlateColor SEnvironmentQueryPin::GetPinColor() const
 
 void SGraphNode_EnvironmentQuery::Construct(const FArguments& InArgs, UEnvironmentQueryGraphNode* InNode)
 {
-	this->GraphNode = InNode;
-
-	this->SetCursor(EMouseCursor::CardinalCross);
-
-	this->UpdateGraphNode();
-
-	bDragMarkerVisible = false;
+	SGraphNodeAI::Construct(SGraphNodeAI::FArguments(), InNode);
 }
 
-void SGraphNode_EnvironmentQuery::AddTest(TSharedPtr<SGraphNode> TestWidget)
+void SGraphNode_EnvironmentQuery::AddSubNode(TSharedPtr<SGraphNode> SubNodeWidget)
 {
+	SGraphNodeAI::AddSubNode(SubNodeWidget);
+
 	TestBox->AddSlot().AutoHeight()
 		[
-			TestWidget.ToSharedRef()
+			SubNodeWidget.ToSharedRef()
 		];
-	TestWidgets.Add(TestWidget);
 }
 
 FSlateColor SGraphNode_EnvironmentQuery::GetBorderBackgroundColor() const
@@ -172,7 +104,7 @@ void SGraphNode_EnvironmentQuery::UpdateGraphNode()
 	// Reset variables that are going to be exposed, in case we are refreshing an already setup node.
 	RightNodeBox.Reset();
 	LeftNodeBox.Reset();
-	TestWidgets.Reset();
+	SubNodes.Reset();
 
 	const float NodePadding = (Cast<UEnvironmentQueryGraphNode_Test>(GraphNode) != NULL) ? 2.0f : 10.0f;
 	float TestPadding = 0.0f;
@@ -180,22 +112,22 @@ void SGraphNode_EnvironmentQuery::UpdateGraphNode()
 	UEnvironmentQueryGraphNode_Option* OptionNode = Cast<UEnvironmentQueryGraphNode_Option>(GraphNode);
 	if (OptionNode)
 	{
-		for (int32 i = 0; i < OptionNode->Tests.Num(); i++)
+		for (int32 Idx = 0; Idx < OptionNode->SubNodes.Num(); Idx++)
 		{
-			if (OptionNode->Tests[i])
+			if (OptionNode->SubNodes[Idx])
 			{
-				TSharedPtr<SGraphNode> NewNode = FNodeFactory::CreateNodeWidget(OptionNode->Tests[i]);
+				TSharedPtr<SGraphNode> NewNode = FNodeFactory::CreateNodeWidget(OptionNode->SubNodes[Idx]);
 				if (OwnerGraphPanelPtr.IsValid())
 				{
 					NewNode->SetOwner(OwnerGraphPanelPtr.Pin().ToSharedRef());
 					OwnerGraphPanelPtr.Pin()->AttachGraphEvents(NewNode);
 				}
-				AddTest(NewNode);
+				AddSubNode(NewNode);
 				NewNode->UpdateGraphNode();
 			}
 		}
 
-		if (TestWidgets.Num() == 0)
+		if (SubNodes.Num() == 0)
 		{
 			TestBox->AddSlot().AutoHeight()
 				[
@@ -361,12 +293,6 @@ void SGraphNode_EnvironmentQuery::UpdateGraphNode()
 	CreatePinWidgets();
 }
 
-FText SGraphNode_EnvironmentQuery::GetDescription() const
-{
-	UEnvironmentQueryGraphNode* StateNode = CastChecked<UEnvironmentQueryGraphNode>(GraphNode);
-	return StateNode ? StateNode->GetDescription() : FText::GetEmpty();
-}
-
 void SGraphNode_EnvironmentQuery::CreatePinWidgets()
 {
 	UEnvironmentQueryGraphNode* StateNode = CastChecked<UEnvironmentQueryGraphNode>(GraphNode);
@@ -377,8 +303,7 @@ void SGraphNode_EnvironmentQuery::CreatePinWidgets()
 		TSharedPtr<SGraphPin> NewPin = SNew(SEnvironmentQueryPin, CurPin);
 
 		NewPin->SetIsEditable(IsEditable);
-
-		this->AddPin(NewPin.ToSharedRef());
+		AddPin(NewPin.ToSharedRef());
 	}
 
 	CurPin = StateNode->GetInputPin();
@@ -387,118 +312,10 @@ void SGraphNode_EnvironmentQuery::CreatePinWidgets()
 		TSharedPtr<SGraphPin> NewPin = SNew(SEnvironmentQueryPin, CurPin);
 
 		NewPin->SetIsEditable(IsEditable);
-
-		this->AddPin(NewPin.ToSharedRef());
+		AddPin(NewPin.ToSharedRef());
 	}
 }
 
-void SGraphNode_EnvironmentQuery::AddPin(const TSharedRef<SGraphPin>& PinToAdd)
-{
-	PinToAdd->SetOwner( SharedThis(this) );
-
-	const UEdGraphPin* PinObj = PinToAdd->GetPinObj();
-	const bool bAdvancedParameter = PinObj && PinObj->bAdvancedView;
-	if(bAdvancedParameter)
-	{
-		PinToAdd->SetVisibility( TAttribute<EVisibility>(PinToAdd, &SGraphPin::IsPinVisibleAsAdvanced) );
-	}
-
-	if (PinToAdd->GetDirection() == EEdGraphPinDirection::EGPD_Input)
-	{
-		LeftNodeBox->AddSlot()
-			.HAlign(HAlign_Fill)
-			.VAlign(VAlign_Fill)
-			.FillHeight(1.0f)
-			[
-				PinToAdd
-			];
-		InputPins.Add(PinToAdd);
-	}
-	else // Direction == EEdGraphPinDirection::EGPD_Output
-	{
-		RightNodeBox->AddSlot()
-			.HAlign(HAlign_Fill)
-			.VAlign(VAlign_Fill)
-			.FillHeight(1.0f)
-			[
-				PinToAdd
-			];
-		OutputPins.Add(PinToAdd);
-	}
-}
-
-FReply SGraphNode_EnvironmentQuery::OnMouseMove(const FGeometry& SenderGeometry, const FPointerEvent& MouseEvent)
-{
-	if (MouseEvent.IsMouseButtonDown(EKeys::LeftMouseButton) && !(GEditor->bIsSimulatingInEditor || GEditor->PlayWorld))
-	{
-		//if we are holding mouse over a subnode
-		UEnvironmentQueryGraphNode_Test* TestNode = Cast<UEnvironmentQueryGraphNode_Test>(GraphNode);
-		if (TestNode && TestNode->ParentNode)
-		{
-			const TSharedRef<SGraphPanel>& Panel = this->GetOwnerPanel().ToSharedRef();
-			const TSharedRef<SGraphNode>& Node = SharedThis(this);
-			return FReply::Handled().BeginDragDrop(FDragNodeEQSTimed::New(Panel, Node));
-		}
-	}
-
-	if (!MouseEvent.IsMouseButtonDown(EKeys::LeftMouseButton) && bDragMarkerVisible)
-	{
-		SetDragMarker(false);
-	}
-
-	return FReply::Unhandled();
-}
-
-FReply SGraphNode_EnvironmentQuery::OnMouseDown(const FGeometry& SenderGeometry, const FPointerEvent& MouseEvent)
-{
-	UEnvironmentQueryGraphNode_Test* TestNode = Cast<UEnvironmentQueryGraphNode_Test>(GraphNode);
-	if (TestNode && TestNode->ParentNode)
-	{
-		GetOwnerPanel()->SelectionManager.ClickedOnNode(GraphNode,MouseEvent);
-		return FReply::Handled();
-	}
-
-	return FReply::Unhandled();
-}
-
-TSharedPtr<SGraphNode> SGraphNode_EnvironmentQuery::GetSubNodeUnderCursor(const FGeometry& WidgetGeometry, const FPointerEvent& MouseEvent)
-{
-	TSharedPtr<SGraphNode> ResultNode;
-
-	// We just need to find the one WidgetToFind among our descendants.
-	TSet< TSharedRef<SWidget> > SubWidgetsSet;
-	for (int32 i=0; i < TestWidgets.Num(); i++)
-	{
-		SubWidgetsSet.Add(TestWidgets[i].ToSharedRef());
-	}
-
-	TMap<TSharedRef<SWidget>, FArrangedWidget> Result;
-	FindChildGeometries(WidgetGeometry, SubWidgetsSet, Result);
-
-	if ( Result.Num() > 0 )
-	{
-		FArrangedChildren ArrangedChildren(EVisibility::Visible);
-		Result.GenerateValueArray( ArrangedChildren.GetInternalArray() );
-		int32 HoveredIndex = SWidget::FindChildUnderMouse( ArrangedChildren, MouseEvent );
-		if ( HoveredIndex != INDEX_NONE )
-		{
-			ResultNode = StaticCastSharedRef<SGraphNode>(ArrangedChildren[HoveredIndex].Widget);
-		}
-	}
-
-	return ResultNode;
-}
-
-void SGraphNode_EnvironmentQuery::SetDragMarker(bool bEnabled)
-{
-	bDragMarkerVisible = bEnabled;
-}
-
-EVisibility SGraphNode_EnvironmentQuery::GetDragOverMarkerVisibility() const
-{
-	return bDragMarkerVisible ? EVisibility::Visible : EVisibility::Collapsed; 
-}
-
 EVisibility SGraphNode_EnvironmentQuery::GetWeightMarkerVisibility() const
 {
 	UEnvironmentQueryGraphNode_Test* MyTestNode = Cast<UEnvironmentQueryGraphNode_Test>(GraphNode);
@@ -536,9 +353,10 @@ void SGraphNode_EnvironmentQuery::OnTestToggleChanged(ECheckBoxState NewState)
 	{
 		MyTestNode->bTestEnabled = (NewState == ECheckBoxState::Checked);
 		
-		if (MyTestNode->ParentNode)
+		UEnvironmentQueryGraphNode_Option* MyParentNode = Cast<UEnvironmentQueryGraphNode_Option>(MyTestNode->ParentNode);
+		if (MyParentNode)
 		{
-			MyTestNode->ParentNode->CalculateWeights();
+			MyParentNode->CalculateWeights();
 		}
 
 		UEnvironmentQueryGraph* MyGraph = Cast<UEnvironmentQueryGraph>(MyTestNode->GetGraph());
@@ -549,162 +367,4 @@ void SGraphNode_EnvironmentQuery::OnTestToggleChanged(ECheckBoxState NewState)
 	}
 }
 
-void SGraphNode_EnvironmentQuery::OnDragEnter( const FGeometry& MyGeometry, const FDragDropEvent& DragDropEvent )
-{
-	// Is someone dragging a node?
-	TSharedPtr<FDragNode> DragConnectionOp = DragDropEvent.GetOperationAs<FDragNode>();
-	if (DragConnectionOp.IsValid())
-	{
-		// Inform the Drag and Drop operation that we are hovering over this node.
-		TSharedPtr<SGraphNode> SubNode = GetSubNodeUnderCursor(MyGeometry, DragDropEvent);
-		DragConnectionOp->SetHoveredNode( SubNode.IsValid() ? SubNode : SharedThis(this) );
-
-		UEnvironmentQueryGraphNode_Test* TestNode = Cast<UEnvironmentQueryGraphNode_Test>(GraphNode);
-		if (DragConnectionOp->IsValidOperation() && TestNode && TestNode->ParentNode)
-		{
-			SetDragMarker(true);
-		}
-	}
-
-	SGraphNode::OnDragEnter(MyGeometry, DragDropEvent);
-}
-
-FReply SGraphNode_EnvironmentQuery::OnDragOver( const FGeometry& MyGeometry, const FDragDropEvent& DragDropEvent )
-{
-	// Is someone dragging a node?
-	TSharedPtr<FDragNode> DragConnectionOp = DragDropEvent.GetOperationAs<FDragNode>();
-	if (DragConnectionOp.IsValid())
-	{
-		// Inform the Drag and Drop operation that we are hovering over this node.
-		TSharedPtr<SGraphNode> SubNode = GetSubNodeUnderCursor(MyGeometry, DragDropEvent);
-		DragConnectionOp->SetHoveredNode( SubNode.IsValid() ? SubNode : SharedThis(this) );
-	}
-	return SGraphNode::OnDragOver(MyGeometry, DragDropEvent);
-}
-
-void SGraphNode_EnvironmentQuery::OnDragLeave( const FDragDropEvent& DragDropEvent )
-{
-	TSharedPtr<FDragNode> DragConnectionOp = DragDropEvent.GetOperationAs<FDragNode>();
-	if (DragConnectionOp.IsValid())
-	{
-		// Inform the Drag and Drop operation that we are not hovering any pins
-		DragConnectionOp->SetHoveredNode( TSharedPtr<SGraphNode>(NULL) );
-	}
-
-	SetDragMarker(false);
-	SGraphNode::OnDragLeave(DragDropEvent);
-}
-
-FReply SGraphNode_EnvironmentQuery::OnDrop( const FGeometry& MyGeometry, const FDragDropEvent& DragDropEvent )
-{
-	SetDragMarker(false);
-
-	TSharedPtr<FDragNodeEQSTimed> DragNodeOp = DragDropEvent.GetOperationAs<FDragNodeEQSTimed>();
-	if (DragNodeOp.IsValid())
-	{
-		if (!DragNodeOp->IsValidOperation())
-		{
-			return FReply::Handled();
-		}
-
-		const float DragTime = float(FPlatformTime::Seconds() - DragNodeOp->StartTime);
-		if (DragTime < 0.5f)
-		{
-			return FReply::Handled();
-		}
-
-		UEnvironmentQueryGraphNode_Option* MyNode = Cast<UEnvironmentQueryGraphNode_Option>(GraphNode);
-		if (MyNode == NULL)
-		{
-			return FReply::Unhandled();
-		}
-
-		const FScopedTransaction Transaction( NSLOCTEXT("UnrealEd", "GraphEd_DragDropNode", "Drag&Drop Node") );
-		bool bReorderOperation = true;
-
-		const TArray< TSharedRef<SGraphNode> >& DraggedNodes = DragNodeOp->GetNodes();
-		for (int32 Idx = 0; Idx < DraggedNodes.Num(); Idx++)
-		{
-			UEnvironmentQueryGraphNode_Test* DraggedTestNode = Cast<UEnvironmentQueryGraphNode_Test>(DraggedNodes[Idx]->GetNodeObj());
-			if (DraggedTestNode && DraggedTestNode->ParentNode)
-			{
-				if (DraggedTestNode->ParentNode != GraphNode)
-				{
-					bReorderOperation = false;
-				}
-
-				DraggedTestNode->ParentNode->Modify();
-				TArray<UEnvironmentQueryGraphNode_Test*> & SubNodes = DraggedTestNode->ParentNode->Tests;
-				for (int32 SubIdx = 0; SubIdx < SubNodes.Num(); SubIdx++)
-				{
-					if (SubNodes[SubIdx] == DraggedTestNode)
-					{
-						SubNodes.RemoveAt(SubIdx);
-					}
-				}
-			}
-		}
-
-		UEnvironmentQueryGraphNode_Test* DropTargetNode = DragNodeOp->GetDropTargetNode();
-		const int32 InsertIndex = MyNode->Tests.IndexOfByKey(DropTargetNode);
-
-		for (int32 Idx = 0; Idx < DraggedNodes.Num(); Idx++)
-		{
-			UEnvironmentQueryGraphNode_Test* DraggedTestNode = Cast<UEnvironmentQueryGraphNode_Test>(DraggedNodes[Idx]->GetNodeObj());
-			DraggedTestNode->Modify();
-			DraggedTestNode->ParentNode = MyNode;
-			MyNode->Modify();
-
-			if (InsertIndex > -1)
-			{
-				MyNode->Tests.Insert(DraggedTestNode, InsertIndex);
-			}
-			else
-			{
-				MyNode->Tests.Add(DraggedTestNode);
-			}
-		}
-
-		if (bReorderOperation)
-		{
-			UpdateGraphNode();
-		}
-		else
-		{
-			MyNode->GetGraph()->NotifyGraphChanged();
-		}
-	}
-
-	return SGraphNode::OnDrop(MyGeometry, DragDropEvent);
-}
-
-TSharedPtr<SToolTip> SGraphNode_EnvironmentQuery::GetComplexTooltip()
-{
-	return NULL;
-}
-
-FText SGraphNode_EnvironmentQuery::GetPreviewCornerText() const
-{
-	UEnvironmentQueryGraphNode* StateNode = CastChecked<UEnvironmentQueryGraphNode>(GraphNode);
-	return FText::FromString(TEXT("Test CornerText"));
-}
-
-const FSlateBrush* SGraphNode_EnvironmentQuery::GetNameIcon() const
-{
-	return FEditorStyle::GetBrush( TEXT("Graph.StateNode.Icon") );
-}
-
-void SGraphNode_EnvironmentQuery::SetOwner(const TSharedRef<SGraphPanel>& OwnerPanel)
-{
-	SGraphNode::SetOwner(OwnerPanel);
-
-	for (auto& ChildWidget : TestWidgets)
-	{
-		if (ChildWidget.IsValid())
-		{
-			ChildWidget->SetOwner(OwnerPanel);
-		}
-	}
-}
-
 #undef LOCTEXT_NAMESPACE
diff --git a/Engine/Source/Editor/EnvironmentQueryEditor/Private/SGraphNode_EnvironmentQuery.h b/Engine/Source/Editor/EnvironmentQueryEditor/Private/SGraphNode_EnvironmentQuery.h
--- a/Engine/Source/Editor/EnvironmentQueryEditor/Private/SGraphNode_EnvironmentQuery.h
+++ b/Engine/Source/Editor/EnvironmentQueryEditor/Private/SGraphNode_EnvironmentQuery.h
@@ -2,25 +2,9 @@
 
 #pragma once
 
-class UEnvironmentQueryGraphNode_Test;
+#include "SGraphNodeAI.h"
 
-class FDragNodeEQSTimed : public FDragNode
-{
-public:
-	DRAG_DROP_OPERATOR_TYPE(FDragNodeEQSTimed, FDragNode)
-
-	static TSharedRef<FDragNodeEQSTimed> New(const TSharedRef<SGraphPanel>& InGraphPanel, const TSharedRef<SGraphNode>& InDraggedNode);
-	static TSharedRef<FDragNodeEQSTimed> New(const TSharedRef<SGraphPanel>& InGraphPanel, const TArray< TSharedRef<SGraphNode> >& InDraggedNodes);
-
-	UEnvironmentQueryGraphNode_Test* GetDropTargetNode() const;
-
-	double StartTime;
-
-protected:
-	typedef FDragNode Super;
-};
-
-class SGraphNode_EnvironmentQuery : public SGraphNode
+class SGraphNode_EnvironmentQuery : public SGraphNodeAI
 {
 public:
 	SLATE_BEGIN_ARGS(SGraphNode_EnvironmentQuery){}
@@ -31,25 +15,10 @@ public:
 	// SGraphNode interface
 	virtual void UpdateGraphNode() override;
 	virtual void CreatePinWidgets() override;
-	virtual void AddPin(const TSharedRef<SGraphPin>& PinToAdd) override;
-	virtual TSharedPtr<SToolTip> GetComplexTooltip() override;
-	virtual void OnDragEnter( const FGeometry& MyGeometry, const FDragDropEvent& DragDropEvent ) override;
-	virtual FReply OnDragOver( const FGeometry& MyGeometry, const FDragDropEvent& DragDropEvent ) override;
-	virtual FReply OnDrop( const FGeometry& MyGeometry, const FDragDropEvent& DragDropEvent ) override;
-	virtual void OnDragLeave( const FDragDropEvent& DragDropEvent ) override;
-	virtual FReply OnMouseMove(const FGeometry& SenderGeometry, const FPointerEvent& MouseEvent) override;
-	virtual void SetOwner(const TSharedRef<SGraphPanel>& OwnerPanel) override;
-
 	// End of SGraphNode interface
 
-	/** handle mouse down on the node */
-	FReply OnMouseDown(const FGeometry& SenderGeometry, const FPointerEvent& MouseEvent);
-
 	/** adds decorator widget inside current node */
-	void AddTest(TSharedPtr<SGraphNode> TestWidget);
-
-	/** gets decorator or service node if one is found under mouse cursor */
-	TSharedPtr<SGraphNode> GetSubNodeUnderCursor(const FGeometry& WidgetGeometry, const FPointerEvent& MouseEvent);
+	virtual void AddSubNode(TSharedPtr<SGraphNode> SubNodeWidget) override;
 
 	EVisibility GetWeightMarkerVisibility() const;
 	TOptional<float> GetWeightProgressBarPercent() const;
@@ -59,22 +28,9 @@ public:
 	ECheckBoxState IsTestToggleChecked() const;
 	void OnTestToggleChanged(ECheckBoxState NewState);
 
-	/** gets drag over marker visibility */
-	EVisibility GetDragOverMarkerVisibility() const;
-
-	/** sets drag marker visible or collapsed on this node */
-	void SetDragMarker(bool bEnabled);
-
 protected:
-	TArray< TSharedPtr<SGraphNode> > TestWidgets;
 	TSharedPtr<SVerticalBox> TestBox;
 
-	uint32 bDragMarkerVisible : 1;
-
 	FSlateColor GetBorderBackgroundColor() const;
 	FSlateColor GetBackgroundColor() const;
-	FText GetDescription() const;
-
-	virtual FText GetPreviewCornerText() const;
-	virtual const FSlateBrush* GetNameIcon() const;
 };
diff --git a/Engine/Source/Editor/EnvironmentQueryEditor/Public/EnvironmentQueryEditorModule.h b/Engine/Source/Editor/EnvironmentQueryEditor/Public/EnvironmentQueryEditorModule.h
--- a/Engine/Source/Editor/EnvironmentQueryEditor/Public/EnvironmentQueryEditorModule.h
+++ b/Engine/Source/Editor/EnvironmentQueryEditor/Public/EnvironmentQueryEditorModule.h
@@ -3,15 +3,14 @@
 #pragma once
 
 #include "ModuleInterface.h"
+#include "AIGraphModule.h"
 #include "Toolkits/IToolkit.h"	// For EAssetEditorMode
 
 DECLARE_LOG_CATEGORY_EXTERN(LogEnvironmentQueryEditor, Log, All);
 
 class IEnvironmentQueryEditor;
 
-/** DataTable Editor module */
-class FEnvironmentQueryEditorModule : public IModuleInterface,
-	public IHasMenuExtensibility, public IHasToolBarExtensibility
+class FEnvironmentQueryEditorModule : public IModuleInterface, public IHasMenuExtensibility, public IHasToolBarExtensibility
 {
 
 public:
@@ -19,22 +18,25 @@ public:
 	virtual void StartupModule() override;
 	virtual void ShutdownModule() override;
 
-	/** Creates an instance of Niagara editor.  Only virtual so that it can be called across the DLL boundary. */
-	virtual TSharedRef<IEnvironmentQueryEditor> CreateEnvironmentQueryEditor( const EToolkitMode::Type Mode, const TSharedPtr< class IToolkitHost >& InitToolkitHost, class UEnvQuery* Script );
+	/** Creates an instance of EQS editor.  Only virtual so that it can be called across the DLL boundary. */
+	virtual TSharedRef<IEnvironmentQueryEditor> CreateEnvironmentQueryEditor( const EToolkitMode::Type Mode, const TSharedPtr< class IToolkitHost >& InitToolkitHost, class UEnvQuery* Query );
 
 	/** Gets the extensibility managers for outside entities to extend static mesh editor's menus and toolbars */
 	virtual TSharedPtr<FExtensibilityManager> GetMenuExtensibilityManager() {return MenuExtensibilityManager;}
 	virtual TSharedPtr<FExtensibilityManager> GetToolBarExtensibilityManager() {return ToolBarExtensibilityManager;}
 
-	/** Niagara Editor app identifier string */
+	/** EQS Editor app identifier string */
 	static const FName EnvironmentQueryEditorAppIdentifier;
 
+	TSharedPtr<struct FGraphNodeClassHelper> GetClassCache() { return ClassCache; }
+
 private:
+
+	TSharedPtr<struct FGraphNodeClassHelper> ClassCache;
+
 	TSharedPtr<FExtensibilityManager> MenuExtensibilityManager;
 	TSharedPtr<FExtensibilityManager> ToolBarExtensibilityManager;
 
 	/** Asset type actions */
 	TSharedPtr<class FAssetTypeActions_EnvironmentQuery> ItemDataAssetTypeActions;
 };
-
-
diff --git a/Engine/Source/Runtime/AIModule/Classes/EnvironmentQuery/EnvQueryGenerator.h b/Engine/Source/Runtime/AIModule/Classes/EnvironmentQuery/EnvQueryGenerator.h
--- a/Engine/Source/Runtime/AIModule/Classes/EnvironmentQuery/EnvQueryGenerator.h
+++ b/Engine/Source/Runtime/AIModule/Classes/EnvironmentQuery/EnvQueryGenerator.h
@@ -2,6 +2,7 @@
 
 #pragma once
 #include "EnvironmentQuery/EnvQueryTypes.h"
+#include "EnvironmentQuery/EnvQueryNode.h"
 #include "DataProviders/AIDataProvider.h"
 #include "EnvQueryGenerator.generated.h"
 
@@ -15,8 +16,8 @@ namespace EnvQueryGeneratorVersion
 	const int32 Latest = DataProviders;
 }
 
-UCLASS(Abstract)
-class AIMODULE_API UEnvQueryGenerator : public UObject
+UCLASS(Abstract, meta=(Category="Generators"))
+class AIMODULE_API UEnvQueryGenerator : public UEnvQueryNode
 {
 	GENERATED_UCLASS_BODY()
 
@@ -27,20 +28,8 @@ class AIMODULE_API UEnvQueryGenerator : public UObject
 	UPROPERTY()
 	TSubclassOf<UEnvQueryItemType> ItemType;
 
-	/** versioning for updating deprecated properties */
-	UPROPERTY()
-	int32 VerNum;
-
 	virtual void GenerateItems(FEnvQueryInstance& QueryInstance) const { checkNoEntry(); }
 
-	/** get description of generator */
-	virtual FText GetDescriptionTitle() const;
-	virtual FText GetDescriptionDetails() const;
-
-#if WITH_EDITOR && USE_EQS_DEBUGGER
-	virtual void PostEditChangeProperty(struct FPropertyChangedEvent& PropertyChangedEvent) override;
-#endif //WITH_EDITOR && USE_EQS_DEBUGGER
-
 	virtual void PostLoad() override;
-	void UpdateGeneratorVersion();
+	void UpdateNodeVersion();
 };
diff --git a/Engine/Source/Runtime/AIModule/Classes/EnvironmentQuery/EnvQueryNode.h b/Engine/Source/Runtime/AIModule/Classes/EnvironmentQuery/EnvQueryNode.h
new file mode 100644
--- /dev/null
+++ b/Engine/Source/Runtime/AIModule/Classes/EnvironmentQuery/EnvQueryNode.h
@@ -0,0 +1,23 @@
+// Copyright 1998-2015 Epic Games, Inc. All Rights Reserved.
+
+#pragma once
+#include "EnvQueryNode.generated.h"
+
+UCLASS(Abstract)
+class AIMODULE_API UEnvQueryNode : public UObject
+{
+	GENERATED_UCLASS_BODY()
+
+	/** Versioning for updating deprecated properties */
+	UPROPERTY()
+	int32 VerNum;
+
+	virtual void UpdateNodeVersion();
+
+	virtual FText GetDescriptionTitle() const;
+	virtual FText GetDescriptionDetails() const;
+
+#if WITH_EDITOR && USE_EQS_DEBUGGER
+	virtual void PostEditChangeProperty(FPropertyChangedEvent& PropertyChangedEvent) override;
+#endif //WITH_EDITOR && USE_EQS_DEBUGGER
+};
diff --git a/Engine/Source/Runtime/AIModule/Classes/EnvironmentQuery/EnvQueryTest.h b/Engine/Source/Runtime/AIModule/Classes/EnvironmentQuery/EnvQueryTest.h
--- a/Engine/Source/Runtime/AIModule/Classes/EnvironmentQuery/EnvQueryTest.h
+++ b/Engine/Source/Runtime/AIModule/Classes/EnvironmentQuery/EnvQueryTest.h
@@ -3,6 +3,7 @@
 #pragma once
 #include "DataProviders/AIDataProvider.h"
 #include "EnvironmentQuery/EnvQueryTypes.h"
+#include "EnvironmentQuery/EnvQueryNode.h"
 #include "EnvQueryTest.generated.h"
 
 class UEnvQueryItemType;
@@ -20,7 +21,7 @@ namespace EnvQueryTestVersion
 }
 
 UCLASS(Abstract)
-class AIMODULE_API UEnvQueryTest : public UObject
+class AIMODULE_API UEnvQueryTest : public UEnvQueryNode
 {
 	GENERATED_UCLASS_BODY()
 
@@ -28,10 +29,6 @@ class AIMODULE_API UEnvQueryTest : public UObject
 	UPROPERTY()
 	int32 TestOrder;
 
-	/** Versioning for updating deprecated properties */
-	UPROPERTY()
-	int32 VerNum;
-
 	/** The purpose of this test.  Should it be used for filtering possible results, scoring them, or both? */
 	UPROPERTY(EditDefaultsOnly, Category=Test)
 	TEnumAsByte<EEnvTestPurpose::Type> TestPurpose;
@@ -164,21 +161,13 @@ public:
 	FORCEINLINE bool IsScoring() const { return (TestPurpose != EEnvTestPurpose::Filter); } 
 	FORCEINLINE bool IsFiltering() const { return (TestPurpose != EEnvTestPurpose::Score); }
 
-	/** get description of test */
-	virtual FString GetDescriptionTitle() const;
-	virtual FText GetDescriptionDetails() const;
-
 	FText DescribeFloatTestParams() const;
 	FText DescribeBoolTestParams(const FString& ConditionDesc) const;
 
 	virtual void PostLoad() override;
 
 	/** update to latest version after spawning */
-	void UpdateTestVersion();
-
-#if WITH_EDITOR && USE_EQS_DEBUGGER
-	virtual void PostEditChangeProperty(FPropertyChangedEvent& PropertyChangedEvent) override;
-#endif //WITH_EDITOR && USE_EQS_DEBUGGER
+	virtual void UpdateNodeVersion() override;
 };
 
 //////////////////////////////////////////////////////////////////////////
diff --git a/Engine/Source/Runtime/AIModule/Classes/EnvironmentQuery/Tests/EnvQueryTest_Distance.h b/Engine/Source/Runtime/AIModule/Classes/EnvironmentQuery/Tests/EnvQueryTest_Distance.h
--- a/Engine/Source/Runtime/AIModule/Classes/EnvironmentQuery/Tests/EnvQueryTest_Distance.h
+++ b/Engine/Source/Runtime/AIModule/Classes/EnvironmentQuery/Tests/EnvQueryTest_Distance.h
@@ -32,6 +32,6 @@ class UEnvQueryTest_Distance : public UEnvQueryTest
 
 	virtual void RunTest(FEnvQueryInstance& QueryInstance) const override;
 
-	virtual FString GetDescriptionTitle() const override;
+	virtual FText GetDescriptionTitle() const override;
 	virtual FText GetDescriptionDetails() const override;
 };
diff --git a/Engine/Source/Runtime/AIModule/Classes/EnvironmentQuery/Tests/EnvQueryTest_Dot.h b/Engine/Source/Runtime/AIModule/Classes/EnvironmentQuery/Tests/EnvQueryTest_Dot.h
--- a/Engine/Source/Runtime/AIModule/Classes/EnvironmentQuery/Tests/EnvQueryTest_Dot.h
+++ b/Engine/Source/Runtime/AIModule/Classes/EnvironmentQuery/Tests/EnvQueryTest_Dot.h
@@ -41,7 +41,7 @@ protected:
 
 	virtual void RunTest(FEnvQueryInstance& QueryInstance) const override;
 
-	virtual FString GetDescriptionTitle() const override;
+	virtual FText GetDescriptionTitle() const override;
 	virtual FText GetDescriptionDetails() const override;
 
 	/** helper function: gather directions from context pairs */
diff --git a/Engine/Source/Runtime/AIModule/Classes/EnvironmentQuery/Tests/EnvQueryTest_GameplayTags.h b/Engine/Source/Runtime/AIModule/Classes/EnvironmentQuery/Tests/EnvQueryTest_GameplayTags.h
--- a/Engine/Source/Runtime/AIModule/Classes/EnvironmentQuery/Tests/EnvQueryTest_GameplayTags.h
+++ b/Engine/Source/Runtime/AIModule/Classes/EnvironmentQuery/Tests/EnvQueryTest_GameplayTags.h
@@ -15,7 +15,6 @@ class UEnvQueryTest_GameplayTags : public UEnvQueryTest
 protected:
 	virtual void RunTest(FEnvQueryInstance& QueryInstance) const override;
 
-	virtual FString GetDescriptionTitle() const override;
 	virtual FText GetDescriptionDetails() const override;
 
 	bool SatisfiesTest(IGameplayTagAssetInterface* ItemGameplayTagAssetInterface) const;
diff --git a/Engine/Source/Runtime/AIModule/Classes/EnvironmentQuery/Tests/EnvQueryTest_Random.h b/Engine/Source/Runtime/AIModule/Classes/EnvironmentQuery/Tests/EnvQueryTest_Random.h
--- a/Engine/Source/Runtime/AIModule/Classes/EnvironmentQuery/Tests/EnvQueryTest_Random.h
+++ b/Engine/Source/Runtime/AIModule/Classes/EnvironmentQuery/Tests/EnvQueryTest_Random.h
@@ -12,6 +12,5 @@ class UEnvQueryTest_Random : public UEnvQueryTest
 protected:
 	virtual void RunTest(FEnvQueryInstance& QueryInstance) const override;
 
-	virtual FString GetDescriptionTitle() const override;
 	virtual FText GetDescriptionDetails() const override;
 };
