fixed restoring EQS test nodes
#ue4

diff --git a/Engine/Source/Editor/EnvironmentQueryEditor/Private/EnvironmentQueryGraph.cpp b/Engine/Source/Editor/EnvironmentQueryEditor/Private/EnvironmentQueryGraph.cpp
--- a/Engine/Source/Editor/EnvironmentQueryEditor/Private/EnvironmentQueryGraph.cpp
+++ b/Engine/Source/Editor/EnvironmentQueryEditor/Private/EnvironmentQueryGraph.cpp
@@ -323,6 +323,7 @@ void UEnvironmentQueryGraph::SpawnMissingNodes()
 				UpdateNodeClassData(TestNode, TestsCopy[SubIdx]->GetClass());
 
 				MyNode->AddSubNode(TestNode, this);
+				TestNode->NodeInstance = TestsCopy[SubIdx];
 			}
 		}
 	}
