moved EQS test scoring preview data to test class, so it doesn't depend on property changed notify for updates
#ue4 UE-11398

diff --git a/Engine/Source/Editor/EnvironmentQueryEditor/Private/DetailCustomizations/EnvQueryTestDetails.cpp b/Engine/Source/Editor/EnvironmentQueryEditor/Private/DetailCustomizations/EnvQueryTestDetails.cpp
--- a/Engine/Source/Editor/EnvironmentQueryEditor/Private/DetailCustomizations/EnvQueryTestDetails.cpp
+++ b/Engine/Source/Editor/EnvironmentQueryEditor/Private/DetailCustomizations/EnvQueryTestDetails.cpp
@@ -233,15 +233,7 @@ void FEnvQueryTestDetails::CustomizeDetails( IDetailLayoutBuilder& DetailLayout
 		SAssignNew(PreviewWidget, STestFunctionWidget)
 	];
 
-	FSimpleDelegate OnGraphPreviewDataChangedDelegate = FSimpleDelegate::CreateSP(this, &FEnvQueryTestDetails::UpdateTestFunctionPreview);
-	TestPurposeHandle->SetOnPropertyValueChanged(OnGraphPreviewDataChangedDelegate);
-	FilterTypeHandle->SetOnPropertyValueChanged(OnGraphPreviewDataChangedDelegate);
-	ClampMaxTypeHandle->SetOnPropertyValueChanged(OnGraphPreviewDataChangedDelegate);
-	ClampMinTypeHandle->SetOnPropertyValueChanged(OnGraphPreviewDataChangedDelegate);
-	ScoreEquationHandle->SetOnPropertyValueChanged(OnGraphPreviewDataChangedDelegate);
-	ScoreHandle->SetOnPropertyValueChanged(OnGraphPreviewDataChangedDelegate);
-
-	UpdateTestFunctionPreview();
+	PreviewWidget->DrawTestOb = Cast<UEnvQueryTest>(MyTest.Get());
 }
 END_SLATE_FUNCTION_BUILD_OPTIMIZATION
 
@@ -749,102 +741,4 @@ EVisibility FEnvQueryTestDetails::GetTestPreviewVisibility() const
 	return EVisibility::Collapsed;
 }
 
-void FEnvQueryTestDetails::UpdateTestFunctionPreview() const
-{
-	// read values of properties
-	uint8 TestPurpose = 0;
-	uint8 FilterType = 0;
-	uint8 ClampMinType = 0;
-	uint8 ClampMaxType = 0;
-	uint8 ScoreEquation = 0;
-	float Score = 0.0f;
-
-	TestPurposeHandle->GetValue(TestPurpose);
-	FilterTypeHandle->GetValue(FilterType);
-	ClampMinTypeHandle->GetValue(ClampMinType);
-	ClampMaxTypeHandle->GetValue(ClampMaxType);
-	ScoreEquationHandle->GetValue(ScoreEquation);
-	ScoreHandle->GetValue(Score);
-
-	// fill score equation samples
-	if (TestPurpose == EEnvTestPurpose::Filter)
-	{
-		// pure filtering won't apply any scoring, draw flat line
-		ScoreEquation = EEnvTestScoreEquation::Constant;
-	}
-
-	FillEquationSamples(ScoreEquation, Score < 0.0f, PreviewWidget->ScoreValues);
-
-	// fill everything else
-	const bool bCanFilter = (TestPurpose != EEnvTestPurpose::Score);
-	PreviewWidget->bShowClampMin = (ClampMinType != EEnvQueryTestClamping::None);
-	PreviewWidget->bShowClampMax = (ClampMaxType != EEnvQueryTestClamping::None);
-	PreviewWidget->bShowLowPassFilter = ((FilterType == EEnvTestFilterType::Minimum) || (FilterType == EEnvTestFilterType::Range)) && bCanFilter;
-	PreviewWidget->bShowHiPassFilter = ((FilterType == EEnvTestFilterType::Maximum) || (FilterType == EEnvTestFilterType::Range)) && bCanFilter;
-	PreviewWidget->FilterLowX = 0.2f;
-	PreviewWidget->FilterHiX = 0.8f;
-	PreviewWidget->ClampMinX = (ClampMinType == EEnvQueryTestClamping::FilterThreshold) ? PreviewWidget->FilterLowX : 0.3f;
-	PreviewWidget->ClampMaxX = (ClampMaxType == EEnvQueryTestClamping::FilterThreshold) ? PreviewWidget->FilterHiX : 0.7f;
-	
-	// postprocess samples for clamping
-	if (PreviewWidget->bShowClampMin)
-	{
-		const int32 FixedIdx = FMath::TruncToInt(PreviewWidget->ClampMinX * 10.0f);
-		for (int32 Idx = 0; Idx < FixedIdx; Idx++)
-		{
-			PreviewWidget->ScoreValues[Idx] = PreviewWidget->ScoreValues[FixedIdx];
-		}
-	}
-	
-	if (PreviewWidget->bShowClampMax)
-	{
-		const int32 FixedIdx = FMath::TruncToInt(PreviewWidget->ClampMaxX * 10.0f) + 1;
-		for (int32 Idx = FixedIdx + 1; Idx < PreviewWidget->ScoreValues.Num(); Idx++)
-		{
-			PreviewWidget->ScoreValues[Idx] = PreviewWidget->ScoreValues[FixedIdx];
-		}
-	}
-}
-
-void FEnvQueryTestDetails::FillEquationSamples(uint8 EquationType, bool bInversed, TArray<float>& Samples) const
-{
-	const int32 MaxSamples = 11;
-	static float SamplesLinear[MaxSamples] = { 0.0f };
-	static float SamplesSquare[MaxSamples] = { 0.0f };
-	static float SamplesConstant[MaxSamples] = { 0.0f };
-	static bool bSamplesInitialized = false;
-
-	if (!bSamplesInitialized)
-	{
-		bSamplesInitialized = true;
-
-		for (int32 Idx = 0; Idx < MaxSamples; Idx++)
-		{
-			const float XValue = 1.0f * Idx / (MaxSamples - 1);
-			SamplesLinear[Idx] = XValue;
-			SamplesSquare[Idx] = XValue * XValue;
-			SamplesConstant[Idx] = 0.5f;				// just for looks on preview, not the actual value
-		}
-	}
-
-	const float* AllSamples[] = { SamplesLinear, SamplesSquare, SamplesLinear, SamplesConstant };
-	if (EquationType >= ARRAY_COUNT(AllSamples))
-	{
-		EquationType = EEnvTestScoreEquation::Constant;
-	}
-
-	const float* SamplesArray = AllSamples[EquationType];
-	if (EquationType == EEnvTestScoreEquation::InverseLinear)
-	{
-		bInversed = !bInversed;
-	}
-
-	Samples.Reset();
-	Samples.AddZeroed(MaxSamples);
-	for (int32 Idx = 0; Idx < MaxSamples; Idx++)
-	{
-		Samples[Idx] = bInversed ? (1.0f - SamplesArray[Idx]) : SamplesArray[Idx];
-	}
-}
-
 #undef LOCTEXT_NAMESPACE
diff --git a/Engine/Source/Editor/EnvironmentQueryEditor/Private/STestFunctionWidget.cpp b/Engine/Source/Editor/EnvironmentQueryEditor/Private/STestFunctionWidget.cpp
--- a/Engine/Source/Editor/EnvironmentQueryEditor/Private/STestFunctionWidget.cpp
+++ b/Engine/Source/Editor/EnvironmentQueryEditor/Private/STestFunctionWidget.cpp
@@ -34,13 +34,21 @@ int32 STestFunctionWidget::OnPaint(const FPaintArgs& Args, const FGeometry& Allo
 
 	LayerId++;
 
+	UEnvQueryTest* TestOb = DrawTestOb.Get();
+	if (TestOb == nullptr)
+	{
+		return LayerId;
+	}
+
+	const FEnvQueryTestScoringPreview& PreviewData = TestOb->PreviewData;
+
 	// Draw filter background
-	if (bShowLowPassFilter)
+	if (PreviewData.bShowFilterLow)
 	{
 		FSlateDrawElement::MakeBox(
 			OutDrawElements,
 			LayerId,
-			AllottedGeometry.ToPaintGeometry(FVector2D(0, 0), FVector2D(FMath::TruncToInt(FilterLowX * AllottedGeometry.Size.X), AllottedGeometry.Size.Y)),
+			AllottedGeometry.ToPaintGeometry(FVector2D(0, 0), FVector2D(FMath::TruncToInt(PreviewData.FilterLow * AllottedGeometry.Size.X), AllottedGeometry.Size.Y)),
 			WhiteBrush,
 			MyClippingRect,
 			DrawEffects,
@@ -50,12 +58,12 @@ int32 STestFunctionWidget::OnPaint(const FPaintArgs& Args, const FGeometry& Allo
 		LayerId++;
 	}
 
-	if (bShowHiPassFilter)
+	if (PreviewData.bShowFilterHigh)
 	{
 		FSlateDrawElement::MakeBox(
 			OutDrawElements,
 			LayerId,
-			AllottedGeometry.ToPaintGeometry(FVector2D(FMath::TruncToInt(FilterHiX * AllottedGeometry.Size.X), 0), FVector2D(AllottedGeometry.Size.X, AllottedGeometry.Size.Y)),
+			AllottedGeometry.ToPaintGeometry(FVector2D(FMath::TruncToInt(PreviewData.FilterHigh * AllottedGeometry.Size.X), 0), FVector2D(AllottedGeometry.Size.X, AllottedGeometry.Size.Y)),
 			WhiteBrush,
 			MyClippingRect,
 			DrawEffects,
@@ -84,11 +92,11 @@ int32 STestFunctionWidget::OnPaint(const FPaintArgs& Args, const FGeometry& Allo
 	LayerId++;
 
 	// Draw clamping and filtering filters
-	if (bShowClampMin)
+	if (PreviewData.bShowClampMin)
 	{
 		TArray<FVector2D> ClampLine;
-		ClampLine.Add(GetWidgetPosition(ClampMinX, 0, AllottedGeometry));
-		ClampLine.Add(GetWidgetPosition(ClampMinX, 1, AllottedGeometry));
+		ClampLine.Add(GetWidgetPosition(PreviewData.ClampMin, 0, AllottedGeometry));
+		ClampLine.Add(GetWidgetPosition(PreviewData.ClampMin, 1, AllottedGeometry));
 
 		FSlateDrawElement::MakeLines(
 			OutDrawElements,
@@ -103,11 +111,11 @@ int32 STestFunctionWidget::OnPaint(const FPaintArgs& Args, const FGeometry& Allo
 		LayerId++;
 	}
 
-	if (bShowClampMax)
+	if (PreviewData.bShowClampMax)
 	{
 		TArray<FVector2D> ClampLine;
-		ClampLine.Add(GetWidgetPosition(ClampMaxX, 0, AllottedGeometry));
-		ClampLine.Add(GetWidgetPosition(ClampMaxX, 1, AllottedGeometry));
+		ClampLine.Add(GetWidgetPosition(PreviewData.ClampMax, 0, AllottedGeometry));
+		ClampLine.Add(GetWidgetPosition(PreviewData.ClampMax, 1, AllottedGeometry));
 
 		FSlateDrawElement::MakeLines(
 			OutDrawElements,
@@ -122,11 +130,11 @@ int32 STestFunctionWidget::OnPaint(const FPaintArgs& Args, const FGeometry& Allo
 		LayerId++;
 	}
 
-	if (bShowLowPassFilter)
+	if (PreviewData.bShowFilterLow)
 	{
 		TArray<FVector2D> FilterLine;
-		FilterLine.Add(GetWidgetPosition(FilterLowX, 0, AllottedGeometry));
-		FilterLine.Add(GetWidgetPosition(FilterLowX, 1, AllottedGeometry));
+		FilterLine.Add(GetWidgetPosition(PreviewData.FilterLow, 0, AllottedGeometry));
+		FilterLine.Add(GetWidgetPosition(PreviewData.FilterLow, 1, AllottedGeometry));
 
 		FSlateDrawElement::MakeLines(
 			OutDrawElements,
@@ -141,11 +149,11 @@ int32 STestFunctionWidget::OnPaint(const FPaintArgs& Args, const FGeometry& Allo
 		LayerId++;
 	}
 
-	if (bShowHiPassFilter)
+	if (PreviewData.bShowFilterHigh)
 	{
 		TArray<FVector2D> FilterLine;
-		FilterLine.Add(GetWidgetPosition(FilterHiX, 0, AllottedGeometry));
-		FilterLine.Add(GetWidgetPosition(FilterHiX, 1, AllottedGeometry));
+		FilterLine.Add(GetWidgetPosition(PreviewData.FilterHigh, 0, AllottedGeometry));
+		FilterLine.Add(GetWidgetPosition(PreviewData.FilterHigh, 1, AllottedGeometry));
 
 		FSlateDrawElement::MakeLines(
 			OutDrawElements,
@@ -163,11 +171,11 @@ int32 STestFunctionWidget::OnPaint(const FPaintArgs& Args, const FGeometry& Allo
 	// Draw line graph
 	TArray<FVector2D> LinePoints;
 
-	const float DeltaX = ((float)AllottedGeometry.Size.X / (ScoreValues.Num() - 1));
-	for (int32 Idx = 0; Idx < ScoreValues.Num(); Idx++)
+	const float DeltaX = ((float)AllottedGeometry.Size.X / (ARRAY_COUNT(PreviewData.Samples) - 1));
+	for (int32 Idx = 0; Idx < ARRAY_COUNT(PreviewData.Samples); Idx++)
 	{
 		const float XPos = Idx * DeltaX;
-		const float YPos = (AllottedGeometry.Size.Y - 1) - (ScoreValues[Idx] * AllottedGeometry.Size.Y);
+		const float YPos = (AllottedGeometry.Size.Y - 1) - (PreviewData.Samples[Idx] * AllottedGeometry.Size.Y);
 
 		LinePoints.Add(FVector2D(FMath::TruncToInt(XPos), FMath::TruncToInt(YPos)));
 	}
diff --git a/Engine/Source/Editor/EnvironmentQueryEditor/Private/STestFunctionWidget.h b/Engine/Source/Editor/EnvironmentQueryEditor/Private/STestFunctionWidget.h
--- a/Engine/Source/Editor/EnvironmentQueryEditor/Private/STestFunctionWidget.h
+++ b/Engine/Source/Editor/EnvironmentQueryEditor/Private/STestFunctionWidget.h
@@ -2,6 +2,8 @@
 
 #pragma once
 
+class UEnvQueryTest;
+
 /**
 * Debug console widget, designed to be summoned on top of a viewport or window
 */
@@ -18,19 +20,7 @@ public:
 	virtual FVector2D ComputeDesiredSize(float) const override;
 	virtual int32 OnPaint(const FPaintArgs& Args, const FGeometry& AllottedGeometry, const FSlateRect& MyClippingRect, FSlateWindowElementList& OutDrawElements, int32 LayerId, const FWidgetStyle& InWidgetStyle, bool bParentEnabled) const override;
 
-	// normalized test scores for 11 normalized test values (11 samples: y(0), y(0.1), ..., y(1.0))
-	// clamp lines will be fixed at samples x=ClampMinX and x=ClampMaxX
-	TArray<float> ScoreValues;
-	float ClampMinX;
-	float ClampMaxX;
-	float FilterLowX;
-	float FilterHiX;
-
-	// show clamp min line
-	uint32 bShowClampMin : 1;
-	uint32 bShowClampMax : 1;
-	uint32 bShowLowPassFilter : 1;
-	uint32 bShowHiPassFilter : 1;
+	TWeakObjectPtr<UEnvQueryTest> DrawTestOb;
 
 private:
 
diff --git a/Engine/Source/Runtime/AIModule/Classes/EnvironmentQuery/EnvQueryTest.h b/Engine/Source/Runtime/AIModule/Classes/EnvironmentQuery/EnvQueryTest.h
--- a/Engine/Source/Runtime/AIModule/Classes/EnvironmentQuery/EnvQueryTest.h
+++ b/Engine/Source/Runtime/AIModule/Classes/EnvironmentQuery/EnvQueryTest.h
@@ -20,6 +20,21 @@ namespace EnvQueryTestVersion
 	const int32 Latest = DataProviders;
 }
 
+#if WITH_EDITORONLY_DATA
+struct FEnvQueryTestScoringPreview
+{
+	float Samples[11];
+	float FilterLow;
+	float FilterHigh;
+	float ClampMin;
+	float ClampMax;
+	uint32 bShowFilterLow : 1;
+	uint32 bShowFilterHigh : 1;
+	uint32 bShowClampMin : 1;
+	uint32 bShowClampMax : 1;
+};
+#endif
+
 UCLASS(Abstract)
 class AIMODULE_API UEnvQueryTest : public UEnvQueryNode
 {
@@ -81,6 +96,11 @@ class AIMODULE_API UEnvQueryTest : public UEnvQueryNode
 	/** Validation: item type that can be used with this test */
 	TSubclassOf<UEnvQueryItemType> ValidItemType;
 
+#if WITH_EDITORONLY_DATA
+	/** samples of scoring function to show on graph in editor */
+	FEnvQueryTestScoringPreview PreviewData;
+#endif
+
 	void SetWorkOnFloatValues(bool bWorkOnFloats);
 	FORCEINLINE bool GetWorkOnFloatValues() const { return bWorkOnFloatValues; }
 
@@ -148,6 +168,13 @@ public:
 
 	/** update to latest version after spawning */
 	virtual void UpdateNodeVersion() override;
+
+#if WITH_EDITOR
+	virtual void PostEditChangeProperty(FPropertyChangedEvent& PropertyChangedEvent) override;
+#endif
+
+	/** update preview list */
+	void UpdatePreviewData();
 };
 
 //////////////////////////////////////////////////////////////////////////
diff --git a/Engine/Source/Runtime/AIModule/Private/EnvironmentQuery/EnvQueryTest.cpp b/Engine/Source/Runtime/AIModule/Private/EnvironmentQuery/EnvQueryTest.cpp
--- a/Engine/Source/Runtime/AIModule/Private/EnvironmentQuery/EnvQueryTest.cpp
+++ b/Engine/Source/Runtime/AIModule/Private/EnvironmentQuery/EnvQueryTest.cpp
@@ -287,6 +287,7 @@ FText UEnvQueryTest::DescribeBoolTestParams(const FString& ConditionDesc) const
 void UEnvQueryTest::SetWorkOnFloatValues(bool bWorkOnFloats)
 {
 	bWorkOnFloatValues = bWorkOnFloats;
+
 	// Make sure FilterType is set to a valid value.
 	if (bWorkOnFloats)
 	{
@@ -294,6 +295,8 @@ void UEnvQueryTest::SetWorkOnFloatValues(bool bWorkOnFloats)
 		{
 			FilterType = EEnvTestFilterType::Range;
 		}
+
+		ScoringEquation = EEnvTestScoreEquation::Linear;
 	}
 	else
 	{
@@ -305,6 +308,102 @@ void UEnvQueryTest::SetWorkOnFloatValues(bool bWorkOnFloats)
 		// Scoring MUST be Constant for boolean tests.
 		ScoringEquation = EEnvTestScoreEquation::Constant;
 	}
+
+	UpdatePreviewData();
+}
+
+#if WITH_EDITOR
+void UEnvQueryTest::PostEditChangeProperty(FPropertyChangedEvent& PropertyChangedEvent)
+{
+	Super::PostEditChangeProperty(PropertyChangedEvent);
+
+	if (PropertyChangedEvent.Property)
+	{
+		if (PropertyChangedEvent.Property->GetFName() == GET_MEMBER_NAME_CHECKED(UEnvQueryTest, TestPurpose) ||
+			PropertyChangedEvent.Property->GetFName() == GET_MEMBER_NAME_CHECKED(UEnvQueryTest, FilterType) ||
+			PropertyChangedEvent.Property->GetFName() == GET_MEMBER_NAME_CHECKED(UEnvQueryTest, ClampMaxType) ||
+			PropertyChangedEvent.Property->GetFName() == GET_MEMBER_NAME_CHECKED(UEnvQueryTest, ClampMinType) ||
+			PropertyChangedEvent.Property->GetFName() == GET_MEMBER_NAME_CHECKED(UEnvQueryTest, ScoringEquation) ||
+			PropertyChangedEvent.Property->GetFName() == GET_MEMBER_NAME_CHECKED(UEnvQueryTest, ScoringFactor))
+		{
+			UpdatePreviewData();
+		}
+	}
+}
+#endif
+
+void UEnvQueryTest::UpdatePreviewData()
+{
+#if WITH_EDITORONLY_DATA && WITH_EDITOR
+	const int32 MaxSamples = 11;
+	static float SamplesLinear[MaxSamples] = { 0.0f };
+	static float SamplesSquare[MaxSamples] = { 0.0f };
+	static float SamplesConstant[MaxSamples] = { 0.0f };
+	static bool bSamplesInitialized = false;
+
+	if (!bSamplesInitialized)
+	{
+		bSamplesInitialized = true;
+
+		for (int32 Idx = 0; Idx < MaxSamples; Idx++)
+		{
+			const float XValue = 1.0f * Idx / (MaxSamples - 1);
+			SamplesLinear[Idx] = XValue;
+			SamplesSquare[Idx] = XValue * XValue;
+			SamplesConstant[Idx] = 0.5f;				// just for looks on preview, not the actual value
+		}
+	}
+
+	const float* AllSamples[] = { SamplesLinear, SamplesSquare, SamplesLinear, SamplesConstant };
+
+	int32 EquationType = (ScoringEquation >= ARRAY_COUNT(AllSamples)) ? EEnvTestScoreEquation::Constant : ScoringEquation;
+	if (TestPurpose == EEnvTestPurpose::Filter)
+	{
+		// pure filtering won't apply any scoring, draw flat line
+		EquationType = EEnvTestScoreEquation::Constant;
+	}
+
+	const float* SamplesArray = AllSamples[EquationType];
+	bool bInversed = ScoringFactor.GetValue() < 0.0f;
+	if (EquationType == EEnvTestScoreEquation::InverseLinear)
+	{
+		bInversed = !bInversed;
+	}
+
+	for (int32 Idx = 0; Idx < MaxSamples; Idx++)
+	{
+		PreviewData.Samples[Idx] = bInversed ? (1.0f - SamplesArray[Idx]) : SamplesArray[Idx];
+	}
+
+	PreviewData.bShowClampMin = (ClampMinType != EEnvQueryTestClamping::None);
+	PreviewData.bShowClampMax = (ClampMaxType != EEnvQueryTestClamping::None);
+	const bool bCanFilter = (TestPurpose != EEnvTestPurpose::Score);
+	PreviewData.bShowFilterHigh = ((FilterType == EEnvTestFilterType::Maximum) || (FilterType == EEnvTestFilterType::Range)) && bCanFilter;
+	PreviewData.bShowFilterLow = ((FilterType == EEnvTestFilterType::Minimum) || (FilterType == EEnvTestFilterType::Range)) && bCanFilter;
+
+	PreviewData.FilterLow = 0.2f;
+	PreviewData.FilterHigh = 0.8f;
+	PreviewData.ClampMin = (ClampMinType == EEnvQueryTestClamping::FilterThreshold) ? PreviewData.FilterLow : 0.3f;
+	PreviewData.ClampMax = (ClampMaxType == EEnvQueryTestClamping::FilterThreshold) ? PreviewData.FilterHigh : 0.7f;
+
+	if (PreviewData.bShowClampMin)
+	{
+		const int32 FixedIdx = FMath::TruncToInt(PreviewData.ClampMin * 10.0f);
+		for (int32 Idx = 0; Idx < FixedIdx; Idx++)
+		{
+			PreviewData.Samples[Idx] = PreviewData.Samples[FixedIdx];
+		}
+	}
+
+	if (PreviewData.bShowClampMax)
+	{
+		const int32 FixedIdx = FMath::TruncToInt(PreviewData.ClampMax * 10.0f) + 1;
+		for (int32 Idx = FixedIdx + 1; Idx < MaxSamples; Idx++)
+		{
+			PreviewData.Samples[Idx] = PreviewData.Samples[FixedIdx];
+		}
+	}
+#endif
 }
 
 #undef LOCTEXT_NAMESPACE
diff --git a/Engine/Source/Runtime/AIModule/Private/EnvironmentQuery/Tests/EnvQueryTest_Pathfinding.cpp b/Engine/Source/Runtime/AIModule/Private/EnvironmentQuery/Tests/EnvQueryTest_Pathfinding.cpp
--- a/Engine/Source/Runtime/AIModule/Private/EnvironmentQuery/Tests/EnvQueryTest_Pathfinding.cpp
+++ b/Engine/Source/Runtime/AIModule/Private/EnvironmentQuery/Tests/EnvQueryTest_Pathfinding.cpp
@@ -19,6 +19,8 @@ UEnvQueryTest_Pathfinding::UEnvQueryTest_Pathfinding(const FObjectInitializer& O
 	SkipUnreachable.DefaultValue = true;
 	FloatValueMin.DefaultValue = 1000.0f;
 	FloatValueMax.DefaultValue = 1000.0f;
+
+	SetWorkOnFloatValues(TestMode != EEnvTestPathfinding::PathExist);
 }
 
 void UEnvQueryTest_Pathfinding::RunTest(FEnvQueryInstance& QueryInstance) const
