EQS graph restoring code will now be able to deal with missing option nodes
#ue4

diff --git a/Engine/Source/Editor/AIGraph/Classes/AIGraph.h b/Engine/Source/Editor/AIGraph/Classes/AIGraph.h
--- a/Engine/Source/Editor/AIGraph/Classes/AIGraph.h
+++ b/Engine/Source/Editor/AIGraph/Classes/AIGraph.h
@@ -38,4 +38,6 @@ protected:
 
 	virtual void CollectAllNodeInstances(TSet<UObject*>& NodeInstances);
 	virtual bool CanRemoveNestedObject(UObject* TestObject) const;
+
+	UEdGraphPin* FindGraphNodePin(UEdGraphNode* Node, EEdGraphPinDirection Dir);
 };
diff --git a/Engine/Source/Editor/AIGraph/Private/AIGraph.cpp b/Engine/Source/Editor/AIGraph/Private/AIGraph.cpp
--- a/Engine/Source/Editor/AIGraph/Private/AIGraph.cpp
+++ b/Engine/Source/Editor/AIGraph/Private/AIGraph.cpp
@@ -140,6 +140,21 @@ void UAIGraph::RemoveOrphanedNodes()
 	}
 }
 
+UEdGraphPin* UAIGraph::FindGraphNodePin(UEdGraphNode* Node, EEdGraphPinDirection Dir)
+{
+	UEdGraphPin* Pin = nullptr;
+	for (int32 Idx = 0; Idx < Node->Pins.Num(); Idx++)
+	{
+		if (Node->Pins[Idx]->Direction == Dir)
+		{
+			Pin = Node->Pins[Idx];
+			break;
+		}
+	}
+
+	return Pin;
+}
+
 bool UAIGraph::IsLocked() const
 {
 	return bLockUpdates;
diff --git a/Engine/Source/Editor/BehaviorTreeEditor/Private/BehaviorTreeGraph.cpp b/Engine/Source/Editor/BehaviorTreeEditor/Private/BehaviorTreeGraph.cpp
--- a/Engine/Source/Editor/BehaviorTreeEditor/Private/BehaviorTreeGraph.cpp
+++ b/Engine/Source/Editor/BehaviorTreeEditor/Private/BehaviorTreeGraph.cpp
@@ -865,8 +865,8 @@ void UBehaviorTreeGraph::SpawnMissingNodes()
 		UBehaviorTreeGraphNode* SpawnedRootNode = BTGraphHelpers::SpawnMissingGraphNodes(BTAsset, RootNode, this);
 		if (RootNode && SpawnedRootNode)
 		{
-			UEdGraphPin* RootOutPin = BTGraphHelpers::FindGraphNodePin(RootNode, EGPD_Output);
-			UEdGraphPin* SpawnedInPin = BTGraphHelpers::FindGraphNodePin(SpawnedRootNode, EGPD_Input);
+			UEdGraphPin* RootOutPin = FindGraphNodePin(RootNode, EGPD_Output);
+			UEdGraphPin* SpawnedInPin = FindGraphNodePin(SpawnedRootNode, EGPD_Input);
 
 			RootOutPin->MakeLinkTo(SpawnedInPin);
 		}
diff --git a/Engine/Source/Editor/EnvironmentQueryEditor/Classes/EnvironmentQueryGraph.h b/Engine/Source/Editor/EnvironmentQueryEditor/Classes/EnvironmentQueryGraph.h
--- a/Engine/Source/Editor/EnvironmentQueryEditor/Classes/EnvironmentQueryGraph.h
+++ b/Engine/Source/Editor/EnvironmentQueryEditor/Classes/EnvironmentQueryGraph.h
@@ -4,6 +4,10 @@
 #include "AIGraph.h"
 #include "EnvironmentQueryGraph.generated.h"
 
+class UEnvQueryOption;
+class UEnvQueryTest;
+class UEnvironmentQueryGraphNode;
+
 UCLASS()
 class UEnvironmentQueryGraph : public UAIGraph
 {
@@ -29,4 +33,6 @@ protected:
 	void UpdateNodeClassData(UAIGraphNode* UpdateNode, UClass* InstanceClass);
 
 	virtual void CollectAllNodeInstances(TSet<UObject*>& NodeInstances) override;
+
+	void SpawnMissingSubNodes(UEnvQueryOption* Option, TSet<UEnvQueryTest*> ExistingTests, UEnvironmentQueryGraphNode* OptionNode);
 };
diff --git a/Engine/Source/Editor/EnvironmentQueryEditor/Private/EnvironmentQueryGraph.cpp b/Engine/Source/Editor/EnvironmentQueryEditor/Private/EnvironmentQueryGraph.cpp
--- a/Engine/Source/Editor/EnvironmentQueryEditor/Private/EnvironmentQueryGraph.cpp
+++ b/Engine/Source/Editor/EnvironmentQueryEditor/Private/EnvironmentQueryGraph.cpp
@@ -250,8 +250,10 @@ void UEnvironmentQueryGraph::UpdateNodeClassData(UAIGraphNode* UpdateNode, UClas
 	}
 	else
 	{
-		UpdateNode->ClassData = FGraphNodeClassData(InstanceClass, "");
+		UpdateNode->ClassData = FGraphNodeClassData(InstanceClass, FGraphNodeClassHelper::GetDeprecationMessage(InstanceClass));
 	}
+
+	UpdateNode->ErrorMessage = UpdateNode->ClassData.GetDeprecatedMessage();
 }
 
 void UEnvironmentQueryGraph::CollectAllNodeInstances(TSet<UObject*>& NodeInstances)
@@ -284,8 +286,17 @@ void UEnvironmentQueryGraph::UpdateDeprecatedGeneratorClasses()
 
 void UEnvironmentQueryGraph::SpawnMissingNodes()
 {
+	UEnvQuery* QueryOwner = Cast<UEnvQuery>(GetOuter());
+	if (QueryOwner == nullptr)
+	{
+		return;
+	}
+
+	TSet<UEnvQueryTest*> ExistingTests;
 	TSet<UEnvQueryOption*> ExistingNodes;
+	TArray<UEnvQueryOption*> OptionsCopy = QueryOwner->Options;
 
+	UAIGraphNode* MyRootNode = nullptr;
 	for (int32 Idx = 0; Idx < Nodes.Num(); Idx++)
 	{
 		UEnvironmentQueryGraphNode* MyNode = Cast<UEnvironmentQueryGraphNode>(Nodes[Idx]);
@@ -294,7 +305,7 @@ void UEnvironmentQueryGraph::SpawnMissingNodes()
 		{
 			ExistingNodes.Add(OptionInstance);
 
-			TSet<UEnvQueryTest*> ExistingTests;
+			ExistingTests.Empty(ExistingTests.Num());
 			for (int32 SubIdx = 0; SubIdx < MyNode->SubNodes.Num(); SubIdx++)
 			{
 				UEnvironmentQueryGraphNode* MySubNode = Cast<UEnvironmentQueryGraphNode>(MyNode->SubNodes[SubIdx]);
@@ -310,21 +321,64 @@ void UEnvironmentQueryGraph::SpawnMissingNodes()
 				}
 			}
 
-			TArray<UEnvQueryTest*> TestsCopy = OptionInstance->Tests;
-			for (int32 SubIdx = 0; SubIdx < TestsCopy.Num(); SubIdx++)
-			{
-				if (ExistingTests.Contains(TestsCopy[SubIdx]) || (TestsCopy[SubIdx] == nullptr))
-				{
-					continue;
-				}
+			SpawnMissingSubNodes(OptionInstance, ExistingTests, MyNode);
+		}
 
-				UEnvironmentQueryGraphNode_Test* TestNode = NewObject<UEnvironmentQueryGraphNode_Test>(this);
-				TestNode->NodeInstance = TestsCopy[SubIdx];
-				UpdateNodeClassData(TestNode, TestsCopy[SubIdx]->GetClass());
+		UEnvironmentQueryGraphNode_Root* RootNode = Cast<UEnvironmentQueryGraphNode_Root>(Nodes[Idx]);
+		if (RootNode)
+		{
+			MyRootNode = RootNode;
+		}
+	}
 
-				MyNode->AddSubNode(TestNode, this);
-				TestNode->NodeInstance = TestsCopy[SubIdx];
-			}
+	UEdGraphPin* RootOutPin = MyRootNode ? FindGraphNodePin(MyRootNode, EGPD_Output) : nullptr;
+	ExistingTests.Empty(0);
+
+	for (int32 Idx = 0; Idx < OptionsCopy.Num(); Idx++)
+	{
+		UEnvQueryOption* OptionInstance = OptionsCopy[Idx];
+		if (ExistingNodes.Contains(OptionInstance) || OptionInstance == nullptr || OptionInstance->Generator == nullptr)
+		{
+			continue;
 		}
+
+		FGraphNodeCreator<UEnvironmentQueryGraphNode_Option> NodeBuilder(*this);
+		UEnvironmentQueryGraphNode_Option* MyNode = NodeBuilder.CreateNode();
+		UpdateNodeClassData(MyNode, OptionInstance->Generator->GetClass());
+		NodeBuilder.Finalize();
+
+		if (MyRootNode)
+		{
+			MyNode->NodePosX = MyRootNode->NodePosX + (Idx * 300);
+			MyNode->NodePosY = MyRootNode->NodePosY + 100;
+		}
+
+		MyNode->NodeInstance = OptionInstance;
+		SpawnMissingSubNodes(OptionInstance, ExistingTests, MyNode);
+
+		UEdGraphPin* SpawnedInPin = FindGraphNodePin(MyNode, EGPD_Input);
+		if (RootOutPin && SpawnedInPin)
+		{
+			RootOutPin->MakeLinkTo(SpawnedInPin);
+		}
+	}
+}
+
+void UEnvironmentQueryGraph::SpawnMissingSubNodes(UEnvQueryOption* Option, TSet<UEnvQueryTest*> ExistingTests, UEnvironmentQueryGraphNode* OptionNode)
+{
+	TArray<UEnvQueryTest*> TestsCopy = Option->Tests;
+	for (int32 SubIdx = 0; SubIdx < TestsCopy.Num(); SubIdx++)
+	{
+		if (ExistingTests.Contains(TestsCopy[SubIdx]) || (TestsCopy[SubIdx] == nullptr))
+		{
+			continue;
+		}
+
+		UEnvironmentQueryGraphNode_Test* TestNode = NewObject<UEnvironmentQueryGraphNode_Test>(this);
+		TestNode->NodeInstance = TestsCopy[SubIdx];
+		UpdateNodeClassData(TestNode, TestsCopy[SubIdx]->GetClass());
+
+		OptionNode->AddSubNode(TestNode, this);
+		TestNode->NodeInstance = TestsCopy[SubIdx];
 	}
 }
