fixed ParentNode chain in EQS graph
#ue4 UE-11549

diff --git a/Engine/Source/Editor/EnvironmentQueryEditor/Private/EnvironmentQueryGraph.cpp b/Engine/Source/Editor/EnvironmentQueryEditor/Private/EnvironmentQueryGraph.cpp
--- a/Engine/Source/Editor/EnvironmentQueryEditor/Private/EnvironmentQueryGraph.cpp
+++ b/Engine/Source/Editor/EnvironmentQueryEditor/Private/EnvironmentQueryGraph.cpp
@@ -67,7 +67,15 @@ void UEnvironmentQueryGraph::UpdateAsset(int32 UpdateFlags)
 
 					for (int32 TestIdx = 0; TestIdx < OptionNode->SubNodes.Num(); TestIdx++)
 					{
-						UEnvironmentQueryGraphNode_Test* TestNode = Cast<UEnvironmentQueryGraphNode_Test>(OptionNode->SubNodes[TestIdx]);
+						UAIGraphNode* SubNode = OptionNode->SubNodes[TestIdx];
+						if (SubNode == nullptr)
+						{
+							continue;
+						}
+
+						SubNode->ParentNode = OptionNode;
+
+						UEnvironmentQueryGraphNode_Test* TestNode = Cast<UEnvironmentQueryGraphNode_Test>(SubNode);
 						if (TestNode && TestNode->bTestEnabled)
 						{
 							UEnvQueryTest* TestInstance = Cast<UEnvQueryTest>(TestNode->NodeInstance);
