fixed restoring EQS query asset with multiple options
#ue4, merged from dev branch

diff --git a/Engine/Source/Editor/EnvironmentQueryEditor/Private/EnvironmentQueryGraph.cpp b/Engine/Source/Editor/EnvironmentQueryEditor/Private/EnvironmentQueryGraph.cpp
--- a/Engine/Source/Editor/EnvironmentQueryEditor/Private/EnvironmentQueryGraph.cpp
+++ b/Engine/Source/Editor/EnvironmentQueryEditor/Private/EnvironmentQueryGraph.cpp
@@ -103,8 +103,11 @@ void UEnvironmentQueryGraph::UpdateAsset(int32 UpdateFlags)
 void UEnvironmentQueryGraph::Initialize()
 {
 	Super::Initialize();
+	
+	LockUpdates();
 	SpawnMissingNodes();
 	CalculateAllWeights();
+	UnlockUpdates();
 }
 
 void UEnvironmentQueryGraph::OnLoaded()
