fixed crash on pasting composite decorator node in behavior tree editor
#ue4

diff --git a/Engine/Source/Editor/AIGraph/Classes/AIGraph.h b/Engine/Source/Editor/AIGraph/Classes/AIGraph.h
--- a/Engine/Source/Editor/AIGraph/Classes/AIGraph.h
+++ b/Engine/Source/Editor/AIGraph/Classes/AIGraph.h
@@ -24,6 +24,7 @@ class AIGRAPH_API UAIGraph : public UEdGraph
 	bool UpdateUnknownNodeClasses();
 	void UpdateDeprecatedClasses();
 	void RemoveOrphanedNodes();
+	void UpdateClassData();
 
 	bool IsLocked() const;
 	void LockUpdates();
diff --git a/Engine/Source/Editor/AIGraph/Classes/AIGraphNode.h b/Engine/Source/Editor/AIGraph/Classes/AIGraphNode.h
--- a/Engine/Source/Editor/AIGraph/Classes/AIGraphNode.h
+++ b/Engine/Source/Editor/AIGraph/Classes/AIGraphNode.h
@@ -96,12 +96,17 @@ class AIGRAPH_API UAIGraphNode : public UEdGraphNode
 	/** reinitialize node instance */
 	virtual bool RefreshNodeClass();
 
+	/** updates ClassData from node instance */
+	virtual void UpdateNodeClassData();
+
 	/** Check if node instance uses blueprint for its implementation */
 	bool UsesBlueprint() const;
 
 	/** check if node has any errors, used for assigning colors on graph */
 	virtual bool HasErrors() const;
 
+	static void UpdateNodeClassDataFrom(UClass* InstanceClass, FGraphNodeClassData& UpdatedData);
+
 protected:
 
 	virtual void ResetNodeOwner();
diff --git a/Engine/Source/Editor/AIGraph/Private/AIGraph.cpp b/Engine/Source/Editor/AIGraph/Private/AIGraph.cpp
--- a/Engine/Source/Editor/AIGraph/Private/AIGraph.cpp
+++ b/Engine/Source/Editor/AIGraph/Private/AIGraph.cpp
@@ -90,6 +90,26 @@ void UAIGraph::UpdateDeprecatedClasses()
 	}
 }
 
+void UAIGraph::UpdateClassData()
+{
+	for (int32 Idx = 0; Idx < Nodes.Num(); Idx++)
+	{
+		UAIGraphNode* Node = Cast<UAIGraphNode>(Nodes[Idx]);
+		if (Node)
+		{
+			Node->UpdateNodeClassData();
+
+			for (int32 SubIdx = 0; SubIdx < Node->SubNodes.Num(); SubIdx++)
+			{
+				if (Node->SubNodes[SubIdx])
+				{
+					Node->UpdateNodeClassData();
+				}
+			}
+		}
+	}
+}
+
 void UAIGraph::CollectAllNodeInstances(TSet<UObject*>& NodeInstances)
 {
 	for (int32 Idx = 0; Idx < Nodes.Num(); Idx++)
diff --git a/Engine/Source/Editor/AIGraph/Private/AIGraphEditor.cpp b/Engine/Source/Editor/AIGraph/Private/AIGraphEditor.cpp
--- a/Engine/Source/Editor/AIGraph/Private/AIGraphEditor.cpp
+++ b/Engine/Source/Editor/AIGraph/Private/AIGraphEditor.cpp
@@ -420,6 +420,8 @@ void FAIGraphEditor::PasteNodesHere(const FVector2D& Location)
 		}
 	}
 
+	EdGraph->UpdateClassData();
+
 	// Update UI
 	EdGraph->UnlockUpdates();
 	CurrentGraphEditor->NotifyGraphChanged();
diff --git a/Engine/Source/Editor/AIGraph/Private/AIGraphNode.cpp b/Engine/Source/Editor/AIGraph/Private/AIGraphNode.cpp
--- a/Engine/Source/Editor/AIGraph/Private/AIGraphNode.cpp
+++ b/Engine/Source/Editor/AIGraph/Private/AIGraphNode.cpp
@@ -383,6 +383,31 @@ bool UAIGraphNode::RefreshNodeClass()
 	return bUpdated;
 }
 
+void UAIGraphNode::UpdateNodeClassData()
+{
+	if (NodeInstance)
+	{
+		UpdateNodeClassDataFrom(NodeInstance->GetClass(), ClassData);
+		ErrorMessage = ClassData.GetDeprecatedMessage();
+	}
+}
+
+void UAIGraphNode::UpdateNodeClassDataFrom(UClass* InstanceClass, FGraphNodeClassData& UpdatedData)
+{
+	if (InstanceClass)
+	{
+		UBlueprint* BPOwner = Cast<UBlueprint>(InstanceClass->ClassGeneratedBy);
+		if (BPOwner)
+		{
+			UpdatedData = FGraphNodeClassData(BPOwner->GetName(), BPOwner->GetOutermost()->GetName(), InstanceClass->GetName(), InstanceClass);
+		}
+		else
+		{
+			UpdatedData = FGraphNodeClassData(InstanceClass, FGraphNodeClassHelper::GetDeprecationMessage(InstanceClass));
+		}
+	}
+}
+
 bool UAIGraphNode::HasErrors() const
 {
 	return ErrorMessage.Len() > 0 || NodeInstance == nullptr;
diff --git a/Engine/Source/Editor/BehaviorTreeEditor/Classes/BehaviorTreeDecoratorGraphNode_Decorator.h b/Engine/Source/Editor/BehaviorTreeEditor/Classes/BehaviorTreeDecoratorGraphNode_Decorator.h
--- a/Engine/Source/Editor/BehaviorTreeEditor/Classes/BehaviorTreeDecoratorGraphNode_Decorator.h
+++ b/Engine/Source/Editor/BehaviorTreeEditor/Classes/BehaviorTreeDecoratorGraphNode_Decorator.h
@@ -26,6 +26,7 @@ class UBehaviorTreeDecoratorGraphNode_Decorator : public UBehaviorTreeDecoratorG
 	virtual void PrepareForCopying() override;
 	void PostCopyNode();
 	bool RefreshNodeClass();
+	void UpdateNodeClassData();
 
 private:
 
diff --git a/Engine/Source/Editor/BehaviorTreeEditor/Classes/BehaviorTreeGraph.h b/Engine/Source/Editor/BehaviorTreeEditor/Classes/BehaviorTreeGraph.h
--- a/Engine/Source/Editor/BehaviorTreeEditor/Classes/BehaviorTreeGraph.h
+++ b/Engine/Source/Editor/BehaviorTreeEditor/Classes/BehaviorTreeGraph.h
@@ -53,4 +53,5 @@ protected:
 	virtual void CollectAllNodeInstances(TSet<UObject*>& NodeInstances);
 
 	void UpdateVersion_UnifiedSubNodes();
+	void UpdateVersion_InnerGraphWhitespace();
 };
diff --git a/Engine/Source/Editor/BehaviorTreeEditor/Classes/BehaviorTreeGraphNode_CompositeDecorator.h b/Engine/Source/Editor/BehaviorTreeEditor/Classes/BehaviorTreeGraphNode_CompositeDecorator.h
--- a/Engine/Source/Editor/BehaviorTreeEditor/Classes/BehaviorTreeGraphNode_CompositeDecorator.h
+++ b/Engine/Source/Editor/BehaviorTreeEditor/Classes/BehaviorTreeGraphNode_CompositeDecorator.h
@@ -33,6 +33,7 @@ class UBehaviorTreeGraphNode_CompositeDecorator : public UBehaviorTreeGraphNode
 	virtual bool IsSubNode() const override;
 	virtual bool HasErrors() const override;
 	virtual bool RefreshNodeClass() override;
+	virtual void UpdateNodeClassData() override;
 
 	virtual void PrepareForCopying() override;
 	virtual void PostCopyNode() override;
diff --git a/Engine/Source/Editor/BehaviorTreeEditor/Private/BehaviorTreeDecoratorGraphNode_Decorator.cpp b/Engine/Source/Editor/BehaviorTreeEditor/Private/BehaviorTreeDecoratorGraphNode_Decorator.cpp
--- a/Engine/Source/Editor/BehaviorTreeEditor/Private/BehaviorTreeDecoratorGraphNode_Decorator.cpp
+++ b/Engine/Source/Editor/BehaviorTreeEditor/Private/BehaviorTreeDecoratorGraphNode_Decorator.cpp
@@ -118,3 +118,11 @@ bool UBehaviorTreeDecoratorGraphNode_Decorator::RefreshNodeClass()
 
 	return bUpdated;
 }
+
+void UBehaviorTreeDecoratorGraphNode_Decorator::UpdateNodeClassData()
+{
+	if (NodeInstance)
+	{
+		UAIGraphNode::UpdateNodeClassDataFrom(NodeInstance->GetClass(), ClassData);
+	}
+}
diff --git a/Engine/Source/Editor/BehaviorTreeEditor/Private/BehaviorTreeGraph.cpp b/Engine/Source/Editor/BehaviorTreeEditor/Private/BehaviorTreeGraph.cpp
--- a/Engine/Source/Editor/BehaviorTreeEditor/Private/BehaviorTreeGraph.cpp
+++ b/Engine/Source/Editor/BehaviorTreeEditor/Private/BehaviorTreeGraph.cpp
@@ -16,8 +16,9 @@ namespace BTGraphVersion
 {
 	const int32 Initial = 0;
 	const int32 UnifiedSubNodes = 1;
+	const int32 InnerGraphWhitespace = 2;
 
-	const int32 Latest = UnifiedSubNodes;
+	const int32 Latest = InnerGraphWhitespace;
 }
 
 UBehaviorTreeGraph::UBehaviorTreeGraph(const FObjectInitializer& ObjectInitializer) : Super(ObjectInitializer)
@@ -1158,6 +1159,11 @@ void UBehaviorTreeGraph::UpdateVersion()
 		UpdateVersion_UnifiedSubNodes();
 	}
 
+	if (GraphVersion < BTGraphVersion::InnerGraphWhitespace)
+	{
+		UpdateVersion_InnerGraphWhitespace();
+	}
+
 	GraphVersion = BTGraphVersion::Latest;
 	Modify();
 }
@@ -1186,3 +1192,25 @@ void UBehaviorTreeGraph::UpdateVersion_UnifiedSubNodes()
 		}
 	}
 }
+
+void UBehaviorTreeGraph::UpdateVersion_InnerGraphWhitespace()
+{
+	for (int32 NodeIdx = 0; NodeIdx < Nodes.Num(); NodeIdx++)
+	{
+		UBehaviorTreeGraphNode* MyNode = Cast<UBehaviorTreeGraphNode>(Nodes[NodeIdx]);
+		if (MyNode)
+		{
+			for (int32 SubIdx = 0; SubIdx < MyNode->SubNodes.Num(); SubIdx++)
+			{
+				UBehaviorTreeGraphNode_CompositeDecorator* InnerGraphNode = Cast<UBehaviorTreeGraphNode_CompositeDecorator>(MyNode->SubNodes[SubIdx]);
+				int32 DummyIdx = INDEX_NONE;
+				
+				if (InnerGraphNode && InnerGraphNode->BoundGraph && InnerGraphNode->BoundGraph->GetName().FindChar(TEXT(' '), DummyIdx))
+				{
+					// don't use white space in name here, it prevents links from being copied correctly
+					InnerGraphNode->BoundGraph->Rename(TEXT("CompositeDecorator"), InnerGraphNode);
+				}
+			}
+		}
+	}
+}
diff --git a/Engine/Source/Editor/BehaviorTreeEditor/Private/BehaviorTreeGraphNode_CompositeDecorator.cpp b/Engine/Source/Editor/BehaviorTreeEditor/Private/BehaviorTreeGraphNode_CompositeDecorator.cpp
--- a/Engine/Source/Editor/BehaviorTreeEditor/Private/BehaviorTreeGraphNode_CompositeDecorator.cpp
+++ b/Engine/Source/Editor/BehaviorTreeEditor/Private/BehaviorTreeGraphNode_CompositeDecorator.cpp
@@ -44,14 +44,19 @@ FText UBehaviorTreeGraphNode_CompositeDecorator::GetDescription() const
 
 void UBehaviorTreeGraphNode_CompositeDecorator::PostPlacedNewNode()
 {
-	CreateBoundGraph();
+	if (BoundGraph == nullptr)
+	{
+		CreateBoundGraph();
+	}
+
 	Super::PostPlacedNewNode();
 }
 
 void UBehaviorTreeGraphNode_CompositeDecorator::PostLoad()
 {
 	Super::PostLoad();
-	if (!BoundGraph)
+
+	if (BoundGraph == nullptr)
 	{
 		CreateBoundGraph();
 	}
@@ -108,6 +113,21 @@ bool UBehaviorTreeGraphNode_CompositeDecorator::RefreshNodeClass()
 	return bUpdated;
 }
 
+void UBehaviorTreeGraphNode_CompositeDecorator::UpdateNodeClassData()
+{
+	if (BoundGraph)
+	{
+		for (int32 i = 0; i < BoundGraph->Nodes.Num(); i++)
+		{
+			UBehaviorTreeDecoratorGraphNode_Decorator* Node = Cast<UBehaviorTreeDecoratorGraphNode_Decorator>(BoundGraph->Nodes[i]);
+			if (Node)
+			{
+				Node->UpdateNodeClassData();
+			}
+		}
+	}
+}
+
 bool UBehaviorTreeGraphNode_CompositeDecorator::HasErrors() const
 {
 	return bHasObserverError;
@@ -118,7 +138,8 @@ void UBehaviorTreeGraphNode_CompositeDecorator::CreateBoundGraph()
 	// Create a new animation graph
 	check(BoundGraph == NULL);
 
-	BoundGraph = FBlueprintEditorUtils::CreateNewGraph(this, TEXT("Composite Decorator"), UBehaviorTreeDecoratorGraph::StaticClass(), UEdGraphSchema_BehaviorTreeDecorator::StaticClass());
+	// don't use white space in name here, it prevents links from being copied correctly
+	BoundGraph = FBlueprintEditorUtils::CreateNewGraph(this, TEXT("CompositeDecorator"), UBehaviorTreeDecoratorGraph::StaticClass(), UEdGraphSchema_BehaviorTreeDecorator::StaticClass());
 	check(BoundGraph);
 
 	// Initialize the anim graph
diff --git a/Engine/Source/Editor/EnvironmentQueryEditor/Classes/EnvironmentQueryGraph.h b/Engine/Source/Editor/EnvironmentQueryEditor/Classes/EnvironmentQueryGraph.h
--- a/Engine/Source/Editor/EnvironmentQueryEditor/Classes/EnvironmentQueryGraph.h
+++ b/Engine/Source/Editor/EnvironmentQueryEditor/Classes/EnvironmentQueryGraph.h
@@ -30,8 +30,6 @@ protected:
 	void UpdateVersion_FixupOuters();
 	void UpdateVersion_CollectClassData();
 
-	void UpdateNodeClassData(UAIGraphNode* UpdateNode, UClass* InstanceClass);
-
 	virtual void CollectAllNodeInstances(TSet<UObject*>& NodeInstances) override;
 
 	void SpawnMissingSubNodes(UEnvQueryOption* Option, TSet<UEnvQueryTest*> ExistingTests, UEnvironmentQueryGraphNode* OptionNode);
diff --git a/Engine/Source/Editor/EnvironmentQueryEditor/Classes/EnvironmentQueryGraphNode_Option.h b/Engine/Source/Editor/EnvironmentQueryEditor/Classes/EnvironmentQueryGraphNode_Option.h
--- a/Engine/Source/Editor/EnvironmentQueryEditor/Classes/EnvironmentQueryGraphNode_Option.h
+++ b/Engine/Source/Editor/EnvironmentQueryEditor/Classes/EnvironmentQueryGraphNode_Option.h
@@ -14,6 +14,7 @@ class UEnvironmentQueryGraphNode_Option : public UEnvironmentQueryGraphNode
 	virtual FText GetDescription() const override;
 
 	virtual void PrepareForCopying() override;
+	virtual void UpdateNodeClassData() override;
 
 	virtual void GetContextMenuActions(const FGraphNodeContextMenuBuilder& Context) const override;
 	void CreateAddTestSubMenu(class FMenuBuilder& MenuBuilder, UEdGraph* Graph) const;
diff --git a/Engine/Source/Editor/EnvironmentQueryEditor/Private/EnvironmentQueryGraph.cpp b/Engine/Source/Editor/EnvironmentQueryEditor/Private/EnvironmentQueryGraph.cpp
--- a/Engine/Source/Editor/EnvironmentQueryEditor/Private/EnvironmentQueryGraph.cpp
+++ b/Engine/Source/Editor/EnvironmentQueryEditor/Private/EnvironmentQueryGraph.cpp
@@ -231,42 +231,7 @@ void UEnvironmentQueryGraph::UpdateVersion_FixupOuters()
 
 void UEnvironmentQueryGraph::UpdateVersion_CollectClassData()
 {
-	for (int32 Idx = 0; Idx < Nodes.Num(); Idx++)
-	{
-		UEnvironmentQueryGraphNode* MyNode = Cast<UEnvironmentQueryGraphNode>(Nodes[Idx]);
-		if (MyNode)
-		{
-			UEnvQueryOption* OptionInstance = Cast<UEnvQueryOption>(MyNode->NodeInstance);
-			if (OptionInstance && OptionInstance->Generator)
-			{
-				UpdateNodeClassData(MyNode, OptionInstance->Generator->GetClass());
-			}
-
-			for (int32 SubIdx = 0; SubIdx < MyNode->SubNodes.Num(); SubIdx++)
-			{
-				UAIGraphNode* SubNode = MyNode->SubNodes[SubIdx];
-				if (SubNode && SubNode->NodeInstance)
-				{
-					UpdateNodeClassData(SubNode, SubNode->NodeInstance->GetClass());
-				}
-			}
-		}
-	}
-}
-
-void UEnvironmentQueryGraph::UpdateNodeClassData(UAIGraphNode* UpdateNode, UClass* InstanceClass)
-{
-	UBlueprint* BPOwner = Cast<UBlueprint>(InstanceClass->ClassGeneratedBy);
-	if (BPOwner)
-	{
-		UpdateNode->ClassData = FGraphNodeClassData(BPOwner->GetName(), BPOwner->GetOutermost()->GetName(), InstanceClass->GetName(), InstanceClass);
-	}
-	else
-	{
-		UpdateNode->ClassData = FGraphNodeClassData(InstanceClass, FGraphNodeClassHelper::GetDeprecationMessage(InstanceClass));
-	}
-
-	UpdateNode->ErrorMessage = UpdateNode->ClassData.GetDeprecatedMessage();
+	UpdateClassData();
 }
 
 void UEnvironmentQueryGraph::CollectAllNodeInstances(TSet<UObject*>& NodeInstances)
@@ -357,7 +322,8 @@ void UEnvironmentQueryGraph::SpawnMissingNodes()
 
 		FGraphNodeCreator<UEnvironmentQueryGraphNode_Option> NodeBuilder(*this);
 		UEnvironmentQueryGraphNode_Option* MyNode = NodeBuilder.CreateNode();
-		UpdateNodeClassData(MyNode, OptionInstance->Generator->GetClass());
+		UAIGraphNode::UpdateNodeClassDataFrom(OptionInstance->Generator->GetClass(), MyNode->ClassData);
+		MyNode->ErrorMessage = MyNode->ClassData.GetDeprecatedMessage();
 		NodeBuilder.Finalize();
 
 		if (MyRootNode)
@@ -389,7 +355,7 @@ void UEnvironmentQueryGraph::SpawnMissingSubNodes(UEnvQueryOption* Option, TSet<
 
 		UEnvironmentQueryGraphNode_Test* TestNode = NewObject<UEnvironmentQueryGraphNode_Test>(this);
 		TestNode->NodeInstance = TestsCopy[SubIdx];
-		UpdateNodeClassData(TestNode, TestsCopy[SubIdx]->GetClass());
+		TestNode->UpdateNodeClassData();
 
 		OptionNode->AddSubNode(TestNode, this);
 		TestNode->NodeInstance = TestsCopy[SubIdx];
diff --git a/Engine/Source/Editor/EnvironmentQueryEditor/Private/EnvironmentQueryGraphNode_Option.cpp b/Engine/Source/Editor/EnvironmentQueryEditor/Private/EnvironmentQueryGraphNode_Option.cpp
--- a/Engine/Source/Editor/EnvironmentQueryEditor/Private/EnvironmentQueryGraphNode_Option.cpp
+++ b/Engine/Source/Editor/EnvironmentQueryEditor/Private/EnvironmentQueryGraphNode_Option.cpp
@@ -65,6 +65,16 @@ void UEnvironmentQueryGraphNode_Option::PrepareForCopying()
 	}
 }
 
+void UEnvironmentQueryGraphNode_Option::UpdateNodeClassData()
+{
+	UEnvQueryOption* OptionInstance = Cast<UEnvQueryOption>(NodeInstance);
+	if (OptionInstance && OptionInstance->Generator)
+	{
+		UpdateNodeClassDataFrom(OptionInstance->Generator->GetClass(), ClassData);
+		ErrorMessage = ClassData.GetDeprecatedMessage();
+	}
+}
+
 FText UEnvironmentQueryGraphNode_Option::GetNodeTitle(ENodeTitleType::Type TitleType) const
 {
 	UEnvQueryOption* OptionInstance = Cast<UEnvQueryOption>(NodeInstance);
