Fix for RootEditor_OnMac and RootEditor_OnLinux nodes' warnings

preflight build in which both nodes completed with no warnings: https://ec-01.epicgames.net/commander/link/jobDetails/jobs/GUBP___depot_UE4_2597596__88293_RunGraph_530809?s=Jobs&tabGroup=stepTableHeader

diff --git a/Engine/Source/Editor/EnvironmentQueryEditor/Private/DetailCustomizations/EnvTraceDataCustomization.cpp b/Engine/Source/Editor/EnvironmentQueryEditor/Private/DetailCustomizations/EnvTraceDataCustomization.cpp
--- a/Engine/Source/Editor/EnvironmentQueryEditor/Private/DetailCustomizations/EnvTraceDataCustomization.cpp
+++ b/Engine/Source/Editor/EnvironmentQueryEditor/Private/DetailCustomizations/EnvTraceDataCustomization.cpp
@@ -150,9 +150,9 @@ void FEnvTraceDataCustomization::OnTraceChannelChanged()
 	if (Result == FPropertyAccess::Success)
 	{
 		ETraceTypeQuery TraceTypeValue = (ETraceTypeQuery)TraceChannelValue;
-		ECollisionChannel TraceChannelValue = UEngineTypes::ConvertToCollisionChannel(TraceTypeValue);
+		ECollisionChannel CollisionChannelValue = UEngineTypes::ConvertToCollisionChannel(TraceTypeValue);
 
-		uint8 SerializedChannelValue = TraceChannelValue;
+		uint8 SerializedChannelValue = CollisionChannelValue;
 		PropTraceChannelSerialized->SetValue(SerializedChannelValue);
 	}
 }
