Merge of 4077826 over from Dev-Framework #UE4

diff --git a/Engine/Source/Editor/BehaviorTreeEditor/Private/AssetTypeActions_BehaviorTree.cpp b/Engine/Source/Editor/BehaviorTreeEditor/Private/AssetTypeActions_BehaviorTree.cpp
--- a/Engine/Source/Editor/BehaviorTreeEditor/Private/AssetTypeActions_BehaviorTree.cpp
+++ b/Engine/Source/Editor/BehaviorTreeEditor/Private/AssetTypeActions_BehaviorTree.cpp
@@ -17,7 +17,7 @@
 
 uint32 FAssetTypeActions_BehaviorTree::GetCategories() 
 { 
-	IAIModule& AIModule = FModuleManager::GetModuleChecked<IAIModule>("AIModule").Get();
+	IAIModule& AIModule = FModuleManager::LoadModuleChecked<IAIModule>("AIModule").Get();
 	return AIModule.GetAIAssetCategoryBit();
 }
 
diff --git a/Engine/Source/Editor/EnvironmentQueryEditor/Private/AssetTypeActions_EnvironmentQuery.cpp b/Engine/Source/Editor/EnvironmentQueryEditor/Private/AssetTypeActions_EnvironmentQuery.cpp
--- a/Engine/Source/Editor/EnvironmentQueryEditor/Private/AssetTypeActions_EnvironmentQuery.cpp
+++ b/Engine/Source/Editor/EnvironmentQueryEditor/Private/AssetTypeActions_EnvironmentQuery.cpp
@@ -32,7 +32,7 @@ UClass* FAssetTypeActions_EnvironmentQuery::GetSupportedClass() const
 
 uint32 FAssetTypeActions_EnvironmentQuery::GetCategories()
 {
-	IAIModule& AIModule = FModuleManager::GetModuleChecked<IAIModule>("AIModule").Get();
+	IAIModule& AIModule = FModuleManager::LoadModuleChecked<IAIModule>("AIModule").Get();
 	return AIModule.GetAIAssetCategoryBit();
 }
 #undef LOCTEXT_NAMESPACE
